Add tool for P4V extensions. Currently supports edigrate, describe, and Horde preflight commands.

diff --git a/Engine/Extras/P4VUtils/Commands/DescribeCommand.cs b/Engine/Extras/P4VUtils/Commands/DescribeCommand.cs
new file mode 100644
--- /dev/null
+++ b/Engine/Extras/P4VUtils/Commands/DescribeCommand.cs
@@ -0,0 +1,36 @@
+// Copyright Epic Games, Inc. All Rights Reserved.
+
+using EpicGames.Core;
+using EpicGames.Perforce;
+using Microsoft.Extensions.Logging;
+using System;
+using System.Collections.Generic;
+using System.ComponentModel;
+using System.Text;
+using System.Threading;
+using System.Threading.Tasks;
+
+namespace P4VUtils.Commands
+{
+	class DescribeCommand : Command
+	{
+		public override string Description => "Prints out a description of the given changelist";
+
+		public override CustomToolInfo CustomTool => new CustomToolInfo("Describe", "%c");
+
+		public override async Task<int> Execute(string[] Args, IReadOnlyDictionary<string, string> ConfigValues, ILogger Logger)
+		{
+			int Change = int.Parse(Args[1]);
+
+			PerforceConnection Perforce = new PerforceConnection(null, null, Logger);
+			DescribeRecord DescribeRecord = await Perforce.DescribeAsync(Change, CancellationToken.None);
+
+			foreach (string Line in DescribeRecord.Description.Split('\n'))
+			{
+				Console.WriteLine(Line);
+			}
+
+			return 0;
+		}
+	}
+}
diff --git a/Engine/Extras/P4VUtils/Commands/EdigrateCommand.cs b/Engine/Extras/P4VUtils/Commands/EdigrateCommand.cs
new file mode 100644
--- /dev/null
+++ b/Engine/Extras/P4VUtils/Commands/EdigrateCommand.cs
@@ -0,0 +1,216 @@
+// Copyright Epic Games, Inc. All Rights Reserved.
+
+using EpicGames.Perforce;
+using Microsoft.Extensions.Logging;
+using System;
+using System.Collections.Generic;
+using System.Linq;
+using System.Text;
+using System.Threading;
+using System.Threading.Tasks;
+
+namespace P4VUtils.Commands
+{
+	class EdigrateCommand : Command
+	{
+		public override string Description => "Converts an integration to an edit";
+
+		public override CustomToolInfo CustomTool => new CustomToolInfo("Convert to Edigrate", "%p");
+
+		public override async Task<int> Execute(string[] Args, IReadOnlyDictionary<string, string> ConfigValues, ILogger Logger)
+		{
+			int Change = int.Parse(Args[1]);
+			bool Debug = Args.Any(x => x.Equals("-Debug", StringComparison.OrdinalIgnoreCase));
+
+			PerforceConnection Perforce = new PerforceConnection(null, null, null, Logger);
+
+			// Get the list of opened files in the CL along with their metadata
+			Logger.LogInformation("Getting the list of files in changelist {Change}...", Change);
+			List<FStatRecord> OpenedRecords = await Perforce.GetOpenFilesAsync(OpenedOptions.None, Change, null, null, -1, null, CancellationToken.None);
+			if (OpenedRecords.Count == 0)
+			{
+				Logger.LogError("Change has no opened files. Aborting.");
+				return 1;
+			}
+
+			// Each record has {depotFile, clientFile, rev, headRev, action, change, type, user, client}
+			Logger.LogInformation("There are {NumFiles} open files in the change.", OpenedRecords.Count);
+
+			// create a mapping of client to local files so we can re-open using local paths (needed for some actions)
+			List<WhereRecord> WhereRecords = await Perforce.WhereAsync(OpenedRecords.Select(x => x.DepotFile!).ToArray(), CancellationToken.None);
+			Dictionary<string, string> OpenedClientToLocalMap = WhereRecords.ToDictionary(x => x.ClientFile, x => x.Path, StringComparer.OrdinalIgnoreCase);
+
+			// If any file that needs to be resolved still, we need to abort.
+			// We don't want to convert non-resolved files to plain edits,
+			// as we need to resolve the merge first!
+			List<ResolveRecord> ResolveRecords = await Perforce.ResolveAsync(Change, ResolveOptions.PreviewOnly, null!, CancellationToken.None);
+			if (ResolveRecords.Any())
+			{
+				Logger.LogError("The following files in the changelist need to be resolved.");
+				Logger.LogError("This script cannot continue until all files are resolved:");
+				foreach (ResolveRecord ResolveRecord in ResolveRecords)
+				{
+					Logger.LogError("    {FileToResolve}", ResolveRecord.ClientFile);
+				}
+				return 1;
+			}
+
+			// We'll just assume all the files in the changelist are integrated. We'll revert and re-open
+			// with the same action/filetype. This should remove the integration records.
+
+			// partition our files into ones we know the actions for, and the ones we don't
+			List<FStatRecord> DeleteActions = new List<FStatRecord>();
+			List<FStatRecord> AddActions = new List<FStatRecord>();
+			List<FStatRecord> EditActions = new List<FStatRecord>();
+			List<FStatRecord> MoveAddActions = new List<FStatRecord>();
+			List<FStatRecord> MoveDeleteActions = new List<FStatRecord>();
+			List<FStatRecord> UnknownActions = new List<FStatRecord>();
+
+			// Multiple action map to the same reopen mapping, so this handles that.
+			Dictionary<FileAction, List<FStatRecord>> ActionMappings = new Dictionary<FileAction, List<FStatRecord>>()
+			{
+				[FileAction.Delete] = DeleteActions,
+				[FileAction.Branch] = AddActions,
+				[FileAction.Add] = AddActions,
+				[FileAction.Edit] = EditActions,
+				[FileAction.Integrate] = EditActions,
+				[FileAction.MoveAdd] = MoveAddActions,
+				[FileAction.MoveDelete] = MoveDeleteActions
+			};
+
+			// this maps all open records to a reopen action.
+			foreach (FStatRecord OpenedRecord in OpenedRecords)
+			{
+				List<FStatRecord>? ActionList;
+				if (!ActionMappings.TryGetValue(OpenedRecord.Action, out ActionList))
+				{
+					ActionList = UnknownActions;
+				}
+				ActionList.Add(OpenedRecord);
+			}
+
+			//  if we have an unknown action, abort with error.
+			if (UnknownActions.Count > 0)
+			{
+				Logger.LogError("These files have unknown actions so cannot be reliably re-opened:");
+				foreach (FStatRecord OpenedRecord in UnknownActions)
+				{
+					Logger.LogError("  {Action}: {ClientFile}", OpenedRecord.Action.ToString(), OpenedRecord.ClientFile);
+				}
+				return 1;
+			}
+
+			// We are good to go, so actually revert the files.
+			Logger.LogInformation("Reverting files that will be reopened");
+			if (!Debug)
+			{
+				PerforceResponseList<RevertRecord> Results = await Perforce.TryRevertAsync(-1, null, RevertOptions.KeepWorkspaceFiles, OpenedRecords.Select(x => x.ClientFile!).ToArray(), CancellationToken.None);
+				Results.EnsureSuccess();
+			}
+
+			// P4 lets you resolve/submit integrations without syncing them to the revision you are resolving against because it's a server operation.
+			// if the haveRev doesn't match the rev of the originally opened file, this was the case.
+			// so we sync -k because the local file matches the one we want, we just have to tell P4 that so it will let us check it out at that revision.
+			// If we simply sync -k to #head we might miss a legitimate submit that happened since our last resolve, and that would stomp that submit.
+			List<FStatRecord> NotSyncedOpenFiles = OpenedRecords.Where(x => x.HaveRevision != x.Revision && x.Action != FileAction.Branch && x.Action != FileAction.Add).ToList();
+			if (NotSyncedOpenFiles.Count > 0)
+			{
+				Logger.LogWarning("The following files were not synced to the revision they were resolved to before. syncing them to the indicated rev first.");
+				foreach (FStatRecord NotSyncedOpenFile in NotSyncedOpenFiles)
+				{
+					Logger.LogWarning("    {File}#{Revision} (had #{HaveRevision})", NotSyncedOpenFile.ClientFile, NotSyncedOpenFile.Revision, NotSyncedOpenFile.HaveRevision);
+				}
+				if (!Debug)
+				{
+					await Perforce.SyncAsync(SyncOptions.KeepWorkspaceFiles, -1, NotSyncedOpenFiles.Select(x => $"{x.ClientFile}#{x.Revision}").ToArray(), CancellationToken.None);
+				}
+			}
+
+			// Perform basic actions
+			await ReopenFiles("delete", AddActions, OpenedClientToLocalMap, x => Perforce.DeleteAsync(Change, DeleteOptions.None, x, CancellationToken.None), Debug, Logger);
+			await ReopenFiles("add", AddActions, OpenedClientToLocalMap, x => Perforce.AddAsync(Change, x, CancellationToken.None), Debug, Logger);
+			await ReopenFiles("edit", AddActions, OpenedClientToLocalMap, x => Perforce.EditAsync(Change, x, CancellationToken.None), Debug, Logger);
+
+			// Perform move actions
+			if (MoveAddActions.Count > 0)
+			{
+				Logger.LogInformation("Re-opening the following files for move:");
+
+				string[] LocalFiles = MoveAddActions.Select(x => OpenedClientToLocalMap[x.ClientFile!]).ToArray();
+				foreach (string LocalFile in LocalFiles)
+				{
+					Logger.LogInformation("    {LocalFile}", LocalFile);
+				}
+
+				// We have to first open the source file for edit (have to use -k because the file has already been moved locally!)
+				if (!Debug)
+				{
+					await Perforce.EditAsync(Change, null, EditOptions.KeepWorkspaceFiles, MoveAddActions.Select(x => x.MovedFile!).ToArray(), CancellationToken.None);
+
+					// then we can open the file for move in the new location (have to use -k because the file has already been moved locally!)
+					foreach ((FStatRecord OpenedRecord, string LocalFile) in MoveAddActions.Zip(LocalFiles))
+					{
+						await Perforce.MoveAsync(Change, null, MoveOptions.KeepWorkspaceFiles, OpenedRecord.MovedFile!, LocalFile, CancellationToken.None);
+					}
+				}
+			}
+
+			// Get the list of reopened files in the CL to check their filetype
+			Logger.LogInformation("Getting the list of files reopened in changelist {Change}...", Change);
+			List<FStatRecord> ReopenedRecords = await Perforce.GetOpenFilesAsync(OpenedOptions.None, Change, null, null, -1, null!, CancellationToken.None);
+			if (ReopenedRecords.Count == 0)
+			{
+				Logger.LogError("Change has no reopened files. This signifies an error in the script! Aborting...");
+				return 1;
+			}
+			if (ReopenedRecords.Count != OpenedRecords.Count)
+			{
+				Logger.LogError("Change doesn't not have the same number of reopened files ({NumReopened}) as originally ({NumOpened}). This probably signifies an error in the script, and the actions should be reverted! Aborting...", ReopenedRecords.Count, OpenedRecords.Count);
+				return 1;
+			}
+
+			// Ensure the filetypes match and ensure each reopened file was in the original changelist.
+			Dictionary<string, FStatRecord> ReopenedRecordsMap = ReopenedRecords.ToDictionary(x => x.ClientFile!, x => x);
+			foreach (FStatRecord OpenedRecord in OpenedRecords)
+			{
+				FStatRecord? ReopenedRecord;
+				if (!ReopenedRecordsMap.TryGetValue(OpenedRecord.ClientFile!, out ReopenedRecord))
+				{
+					Logger.LogError("Could not find original file {ClientFile} in re-opened records. This signifies an error in the script! Aborting...", OpenedRecord.ClientFile);
+					return 1;
+				}
+				if (OpenedRecord.Type != ReopenedRecord.Type)
+				{
+					Logger.LogInformation("Changing filetype of {ClientFile} from {OpenType} to {ReopenType}", OpenedRecord.ClientFile!, ReopenedRecord.Type, OpenedRecord.Type);
+					if (!Debug)
+					{
+						await Perforce.ReopenAsync(Change, OpenedRecord.Type, ReopenedRecord.ClientFile!, CancellationToken.None);
+					}
+				}
+			}
+
+			// Done!
+			Logger.LogInformation("Finished");
+			return 0;
+		}
+
+		async Task ReopenFiles(string Operation, List<FStatRecord> Records, Dictionary<string, string> OpenedClientToLocalMap, Func<string[], Task> ReopenAsync, bool Debug, ILogger Logger)
+		{
+			if (Records.Count > 0)
+			{
+				Logger.LogInformation("Re-opening the following files for {Operation}:", Operation);
+
+				string[] LocalFiles = Records.Select(x => OpenedClientToLocalMap[x.ClientFile!]).ToArray();
+				foreach (string LocalFile in LocalFiles)
+				{
+					Logger.LogInformation("    {LocalFile}", LocalFiles);
+				}
+
+				if (!Debug)
+				{
+					await ReopenAsync(LocalFiles);
+				}
+			}
+		}
+	}
+}
diff --git a/Engine/Extras/P4VUtils/Commands/PreflightCommand.cs b/Engine/Extras/P4VUtils/Commands/PreflightCommand.cs
new file mode 100644
--- /dev/null
+++ b/Engine/Extras/P4VUtils/Commands/PreflightCommand.cs
@@ -0,0 +1,79 @@
+// Copyright Epic Games, Inc. All Rights Reserved.
+
+using EpicGames.Core;
+using EpicGames.Perforce;
+using Microsoft.Extensions.Logging;
+using System;
+using System.Collections.Generic;
+using System.ComponentModel;
+using System.Diagnostics;
+using System.IO;
+using System.Reflection;
+using System.Runtime.InteropServices;
+using System.Text;
+using System.Text.RegularExpressions;
+using System.Threading;
+using System.Threading.Tasks;
+
+namespace P4VUtils.Commands
+{
+	class PreflightCommand : Command
+	{
+		public override string Description => "Runs a preflight of the given changelist on Horde";
+
+		public override CustomToolInfo CustomTool => new CustomToolInfo("Preflight...", "%p");
+
+		public override async Task<int> Execute(string[] Args, IReadOnlyDictionary<string, string> ConfigValues, ILogger Logger)
+		{
+			int Change = int.Parse(Args[1]);
+
+			PerforceConnection Perforce = new PerforceConnection(null, null, Logger);
+
+			ClientRecord Client = await Perforce.GetClientAsync(null, CancellationToken.None);
+
+			string Url = GetUrl(Client.Stream!, Change, ConfigValues);
+			Logger.LogDebug("Opening {Url}", Url);
+			OpenUrl(Url);
+
+			return 0;
+		}
+
+		public virtual string GetUrl(string Stream, int Change, IReadOnlyDictionary<string, string> ConfigValues)
+		{
+			string? BaseUrl;
+			if (!ConfigValues.TryGetValue("HordeServer", out BaseUrl))
+			{
+				BaseUrl = "https://configure-server-url-in-p4vutils.ini";
+			}
+			return $"{BaseUrl.TrimEnd('/')}/preflight?stream={Stream}&change={Change}";
+		}
+
+		void OpenUrl(string Url)
+		{
+			if (RuntimeInformation.IsOSPlatform(OSPlatform.Windows))
+			{
+				Process.Start(new ProcessStartInfo(Url) { UseShellExecute = true });
+			}
+			else if (RuntimeInformation.IsOSPlatform(OSPlatform.Linux))
+			{
+				Process.Start("xdg-open", Url);
+			}
+			else if (RuntimeInformation.IsOSPlatform(OSPlatform.OSX))
+			{
+				Process.Start("open", Url);
+			}
+		}
+	}
+
+	class PreflightAndSubmitCommand : PreflightCommand
+	{
+		public override string Description => "Runs a preflight of the given changelist on Horde and submits it";
+
+		public override CustomToolInfo CustomTool => new CustomToolInfo("Preflight and submit", "%p");
+
+		public override string GetUrl(string Stream, int Change, IReadOnlyDictionary<string, string> ConfigValues)
+		{
+			return base.GetUrl(Stream, Change, ConfigValues) + "&defaulttemplate=true&submit=true";
+		}
+	}
+}
diff --git a/Engine/Extras/P4VUtils/P4VUtils.csproj b/Engine/Extras/P4VUtils/P4VUtils.csproj
new file mode 100644
--- /dev/null
+++ b/Engine/Extras/P4VUtils/P4VUtils.csproj
@@ -0,0 +1,32 @@
+<Project Sdk="Microsoft.NET.Sdk">
+
+  <PropertyGroup>
+    <OutputType>Exe</OutputType>
+    <TargetFramework>netcoreapp3.1</TargetFramework>
+    <Nullable>enable</Nullable>
+  </PropertyGroup>
+
+  <ItemGroup Condition="Exists('..\..\Restricted\NotForLicensees\Extras\P4VUtils\P4VUtils.ini')">
+    <None Include="..\..\Restricted\NotForLicensees\Extras\P4VUtils\P4VUtils.ini" Link="NotForLicensees\P4VUtils.ini">
+      <CopyToOutputDirectory>PreserveNewest</CopyToOutputDirectory>
+    </None>
+  </ItemGroup>
+
+  <ItemGroup>
+    <PackageReference Include="Microsoft.Extensions.Logging.Console" Version="5.0.0" />
+  </ItemGroup>
+
+  <ItemGroup>
+    <ProjectReference Include="..\..\Source\Programs\Shared\EpicGames.Core\EpicGames.Core.csproj" />
+    <ProjectReference Include="..\..\Source\Programs\Shared\EpicGames.Perforce\EpicGames.Perforce.csproj" />
+  </ItemGroup>
+
+  <ItemGroup>
+    <None Update="P4VUtils.ini">
+      <CopyToOutputDirectory>PreserveNewest</CopyToOutputDirectory>
+    </None>
+  </ItemGroup>
+
+  
+
+</Project>
diff --git a/Engine/Extras/P4VUtils/P4VUtils.ini b/Engine/Extras/P4VUtils/P4VUtils.ini
new file mode 100644
--- /dev/null
+++ b/Engine/Extras/P4VUtils/P4VUtils.ini
@@ -0,0 +1 @@
+HordeServer=https://my-horde-server/
diff --git a/Engine/Extras/P4VUtils/Program.cs b/Engine/Extras/P4VUtils/Program.cs
new file mode 100644
--- /dev/null
+++ b/Engine/Extras/P4VUtils/Program.cs
@@ -0,0 +1,239 @@
+// Copyright Epic Games, Inc. All Rights Reserved.
+
+using EpicGames.Core;
+using Microsoft.Extensions.Logging;
+using Microsoft.Extensions.Logging.Console;
+using P4VUtils.Commands;
+using System;
+using System.Collections.Generic;
+using System.ComponentModel;
+using System.IO;
+using System.Linq;
+using System.Reflection;
+using System.Threading.Tasks;
+using System.Xml;
+
+namespace P4VUtils
+{
+	class CustomToolInfo
+	{
+		public string Name { get; set; }
+		public string Arguments { get; set; }
+		public bool AddToContextMenu { get; set; } = true;
+
+		public CustomToolInfo(string Name, string Arguments)
+		{
+			this.Name = Name;
+			this.Arguments = Arguments;
+		}
+	}
+
+	abstract class Command
+	{
+		public abstract string Description { get; }
+
+		public abstract CustomToolInfo CustomTool { get; }
+
+		public abstract Task<int> Execute(string[] Args, IReadOnlyDictionary<string, string> ConfigValues, ILogger Logger);
+	}
+
+	class Program
+	{
+		public static IReadOnlyDictionary<string, Command> Commands { get; } = new Dictionary<string, Command>
+		{
+			["preflight"] = new PreflightCommand(),
+			["preflightandsubmit"] = new PreflightAndSubmitCommand(),
+			["edigrate"] = new EdigrateCommand(),
+			["describe"] = new DescribeCommand()
+		};
+
+		static void PrintHelp(ILogger Logger)
+		{
+			Logger.LogInformation("P4VUtils");
+			Logger.LogInformation("Provides useful shortcuts for working with P4V");
+			Logger.LogInformation("");
+			Logger.LogInformation("Usage:");
+			Logger.LogInformation("  P4VUtils [Command] [Arguments...]");
+			Logger.LogInformation("");
+
+			List<KeyValuePair<string, string>> Table = new List<KeyValuePair<string, string>>();
+			foreach (KeyValuePair<string, Command> Pair in Commands)
+			{
+				Table.Add(new KeyValuePair<string, string>(Pair.Key, Pair.Value.Description));
+			}
+
+			Logger.LogInformation("Commands:");
+			HelpUtils.PrintTable(Table, 2, 15);
+		}
+
+		static async Task<int> Main(string[] Args)
+		{
+			using ILoggerFactory Factory = LoggerFactory.Create(Builder => Builder.AddEpicDefault());//.AddSimpleConsole(Options => { Options.SingleLine = true; Options.IncludeScopes = false; }));
+			ILogger Logger = Factory.CreateLogger<Program>();
+			Log.Logger = Logger;
+
+			if (Args.Length == 0 || Args[0].Equals("-help", StringComparison.OrdinalIgnoreCase))
+			{
+				PrintHelp(Logger);
+				return 0;
+			}
+			else if (Args[0].StartsWith("-"))
+			{
+				Console.WriteLine("Missing command name");
+				PrintHelp(Logger);
+				return 1;
+			}
+			else if (Args[0].Equals("install", StringComparison.OrdinalIgnoreCase))
+			{
+				Logger.LogInformation("Adding custom tools...");
+				return await UpdateCustomToolRegistration(true, Logger);
+			}
+			else if (Args[0].Equals("uninstall", StringComparison.OrdinalIgnoreCase))
+			{
+				Logger.LogInformation("Removing custom tools...");
+				return await UpdateCustomToolRegistration(false, Logger);
+			}
+			else if (Commands.TryGetValue(Args[0], out Command? Command))
+			{
+				if (Args.Any(x => x.Equals("-help", StringComparison.OrdinalIgnoreCase)))
+				{
+					List<KeyValuePair<string, string>> Parameters = CommandLineArguments.GetParameters(Command.GetType());
+					HelpUtils.PrintHelp(Args[0], Command.GetType());
+					return 0;
+				}
+
+				Dictionary<string, string> ConfigValues = ReadConfig();
+				return await Command.Execute(Args, ConfigValues, Logger);
+			}
+			else
+			{
+				Console.WriteLine("Unknown command: {0}", Args[0]);
+				PrintHelp(Logger);
+				return 1;
+			}
+		}
+
+		static Dictionary<string, string> ReadConfig()
+		{
+			Dictionary<string, string> ConfigValues = new Dictionary<string, string>(StringComparer.OrdinalIgnoreCase);
+
+			string BasePath = Path.GetDirectoryName(Assembly.GetExecutingAssembly().Location)!;
+			AppendConfig(Path.Combine(BasePath, "P4VUtils.ini"), ConfigValues);
+			AppendConfig(Path.Combine(BasePath, "NotForLicensees", "P4VUtils.ini"), ConfigValues);
+
+			return ConfigValues;
+		}
+
+		static void AppendConfig(string SourcePath, Dictionary<string, string> ConfigValues)
+		{
+			if (File.Exists(SourcePath))
+			{
+				string[] Lines = File.ReadAllLines(SourcePath);
+				foreach (string Line in Lines)
+				{
+					int EqualsIdx = Line.IndexOf('=');
+					if (EqualsIdx != -1)
+					{
+						string Key = Line.Substring(0, EqualsIdx).Trim();
+						string Value = Line.Substring(EqualsIdx + 1).Trim();
+						ConfigValues[Key] = Value;
+					}
+				}
+			}
+		}
+
+		public static async Task<int> UpdateCustomToolRegistration(bool bInstall, ILogger Logger)
+		{
+			DirectoryReference? ConfigDir = DirectoryReference.GetSpecialFolder(Environment.SpecialFolder.UserProfile);
+			if (ConfigDir == null)
+			{
+				Logger.LogError("Unable to find config directory.");
+				return 1;
+			}
+
+			FileReference ConfigFile = FileReference.Combine(ConfigDir, ".p4qt", "customtools.xml");
+			if (!FileReference.Exists(ConfigFile))
+			{
+				DirectoryReference.CreateDirectory(ConfigFile.Directory);
+				using (StreamWriter Writer = new StreamWriter(ConfigFile.FullName))
+				{
+					await Writer.WriteLineAsync(@"<?xml version=""1.0"" encoding=""UTF-8""?>");
+					await Writer.WriteLineAsync(@"<!--perforce-xml-version=1.0-->");
+					await Writer.WriteLineAsync(@"<CustomToolDefList varName=""customtooldeflist"">");
+					await Writer.WriteLineAsync(@"</CustomToolDefList>");
+				}
+			}
+
+			XmlDocument Document = new XmlDocument();
+			Document.Load(ConfigFile.FullName);
+
+			XmlElement? Root = Document.SelectSingleNode("CustomToolDefList") as XmlElement;
+			if (Root == null)
+			{
+				Logger.LogError("Unknown schema for {ConfigFile}", ConfigFile);
+				return 1;
+			}
+
+			FileReference DotNetLocation = FileReference.Combine(DirectoryReference.GetSpecialFolder(Environment.SpecialFolder.ProgramFiles), "dotnet", "dotnet.exe");
+			FileReference AssemblyLocation = new FileReference(Assembly.GetExecutingAssembly().GetOriginalLocation());
+
+			foreach (XmlNode? ChildNode in Root.SelectNodes("CustomToolDef"))
+			{
+				XmlElement? ChildElement = ChildNode as XmlElement;
+				if (ChildElement != null)
+				{
+					XmlElement? CommandElement = ChildElement.SelectSingleNode("Definition/Command") as XmlElement;
+					if (CommandElement != null && new FileReference(CommandElement.InnerText) == DotNetLocation)
+					{
+						XmlElement? ArgumentsElement = ChildElement.SelectSingleNode("Definition/Arguments") as XmlElement;
+						if (ArgumentsElement != null)
+						{
+							string[] Arguments = CommandLineArguments.Split(ArgumentsElement.InnerText);
+							if (Arguments.Length > 0 && new FileReference(Arguments[0]) == AssemblyLocation)
+							{
+								Root.RemoveChild(ChildElement);
+							}
+						}
+					}
+				}
+			}
+
+			// Insert new entries
+			if (bInstall)
+			{
+				foreach (KeyValuePair<string, Command> Pair in Commands)
+				{
+					CustomToolInfo CustomTool = Pair.Value.CustomTool;
+
+					XmlElement ToolDef = Document.CreateElement("CustomToolDef");
+					{
+						XmlElement Definition = Document.CreateElement("Definition");
+						{
+							XmlElement Description = Document.CreateElement("Name");
+							Description.InnerText = CustomTool.Name;
+							Definition.AppendChild(Description);
+
+							XmlElement Command = Document.CreateElement("Command");
+							Command.InnerText = DotNetLocation.FullName;
+							Definition.AppendChild(Command);
+
+							XmlElement Arguments = Document.CreateElement("Arguments");
+							Arguments.InnerText = $"\"{AssemblyLocation.FullName}\" {Pair.Key} {CustomTool.Arguments}";
+							Definition.AppendChild(Arguments);
+						}
+						ToolDef.AppendChild(Definition);
+
+						XmlElement AddToContext = Document.CreateElement("AddToContext");
+						AddToContext.InnerText = CustomTool.AddToContextMenu ? "true" : "false";
+						ToolDef.AppendChild(AddToContext);
+					}
+					Root.AppendChild(ToolDef);
+				}
+			}
+
+			// Save the new document
+			Document.Save(ConfigFile.FullName);
+			return 0;
+		}
+	}
+}
diff --git a/Engine/Extras/P4VUtils/Properties/launchSettings.json b/Engine/Extras/P4VUtils/Properties/launchSettings.json
new file mode 100644
--- /dev/null
+++ b/Engine/Extras/P4VUtils/Properties/launchSettings.json
@@ -0,0 +1,8 @@
+{
+  "profiles": {
+    "P4VUtils": {
+      "commandName": "Project",
+      "commandLineArgs": "preflight 1234"
+    }
+  }
+}
\ No newline at end of file
