P4VUtils: Handle corrupt XML files gracefully.

#rnx

diff --git a/Engine/Extras/P4VUtils/Program.cs b/Engine/Extras/P4VUtils/Program.cs
--- a/Engine/Extras/P4VUtils/Program.cs
+++ b/Engine/Extras/P4VUtils/Program.cs
@@ -20,6 +20,7 @@ namespace P4VUtils
 		public string Name { get; set; }
 		public string Arguments { get; set; }
 		public bool AddToContextMenu { get; set; } = true;
+		public bool ShowConsole { get; set; }
 
 		public CustomToolInfo(string Name, string Arguments)
 		{
@@ -72,6 +73,19 @@ namespace P4VUtils
 			ILogger Logger = Factory.CreateLogger<Program>();
 			Log.Logger = Logger;
 
+			try
+			{
+				return await InnerMain(Args, Logger);
+			}
+			catch (Exception Ex)
+			{
+				Logger.LogError(Ex, "Unhandled exception: {Ex}", Ex.ToString());
+				return 1;
+			}
+		}
+
+		static async Task<int> InnerMain(string[] Args, ILogger Logger)
+		{
 			if (Args.Length == 0 || Args[0].Equals("-help", StringComparison.OrdinalIgnoreCase))
 			{
 				PrintHelp(Logger);
@@ -142,6 +156,22 @@ namespace P4VUtils
 			}
 		}
 
+		public static bool TryLoadXmlDocument(FileReference Location, XmlDocument Document)
+		{
+			if (FileReference.Exists(Location))
+			{
+				try
+				{
+					Document.Load(Location.FullName);
+					return true;
+				}
+				catch
+				{
+				}
+			}
+			return false;
+		}
+
 		public static async Task<int> UpdateCustomToolRegistration(bool bInstall, ILogger Logger)
 		{
 			DirectoryReference? ConfigDir = DirectoryReference.GetSpecialFolder(Environment.SpecialFolder.UserProfile);
@@ -152,7 +182,9 @@ namespace P4VUtils
 			}
 
 			FileReference ConfigFile = FileReference.Combine(ConfigDir, ".p4qt", "customtools.xml");
-			if (!FileReference.Exists(ConfigFile))
+
+			XmlDocument Document = new XmlDocument();
+			if (!TryLoadXmlDocument(ConfigFile, Document))
 			{
 				DirectoryReference.CreateDirectory(ConfigFile.Directory);
 				using (StreamWriter Writer = new StreamWriter(ConfigFile.FullName))
@@ -162,11 +194,9 @@ namespace P4VUtils
 					await Writer.WriteLineAsync(@"<CustomToolDefList varName=""customtooldeflist"">");
 					await Writer.WriteLineAsync(@"</CustomToolDefList>");
 				}
+				Document.Load(ConfigFile.FullName);
 			}
 
-			XmlDocument Document = new XmlDocument();
-			Document.Load(ConfigFile.FullName);
-
 			XmlElement? Root = Document.SelectSingleNode("CustomToolDefList") as XmlElement;
 			if (Root == null)
 			{
@@ -218,11 +248,22 @@ namespace P4VUtils
 							Definition.AppendChild(Command);
 
 							XmlElement Arguments = Document.CreateElement("Arguments");
-							Arguments.InnerText = $"\"{AssemblyLocation.FullName}\" {Pair.Key} {CustomTool.Arguments}";
+							Arguments.InnerText = $"{AssemblyLocation.FullName.QuoteArgument()} {Pair.Key} {CustomTool.Arguments}";
 							Definition.AppendChild(Arguments);
 						}
 						ToolDef.AppendChild(Definition);
 
+						if (CustomTool.ShowConsole)
+						{
+							XmlElement Console = Document.CreateElement("Console");
+							{
+								XmlElement CloseOnExit = Document.CreateElement("CloseOnExit");
+								CloseOnExit.InnerText = "false";
+								Console.AppendChild(CloseOnExit);
+							}
+							ToolDef.AppendChild(Console);
+						}
+
 						XmlElement AddToContext = Document.CreateElement("AddToContext");
 						AddToContext.InnerText = CustomTool.AddToContextMenu ? "true" : "false";
 						ToolDef.AppendChild(AddToContext);
@@ -233,6 +274,7 @@ namespace P4VUtils
 
 			// Save the new document
 			Document.Save(ConfigFile.FullName);
+			Logger.LogInformation("Written {ConfigFile}", ConfigFile.FullName);
 			return 0;
 		}
 	}
diff --git a/Engine/Source/Programs/Shared/EpicGames.Perforce/PerforceConnection.cs b/Engine/Source/Programs/Shared/EpicGames.Perforce/PerforceConnection.cs
--- a/Engine/Source/Programs/Shared/EpicGames.Perforce/PerforceConnection.cs
+++ b/Engine/Source/Programs/Shared/EpicGames.Perforce/PerforceConnection.cs
@@ -959,7 +959,7 @@ namespace EpicGames.Perforce
 			{
 				Arguments.Append(" -v");
 			}
-			Arguments.Append(StreamName);
+			Arguments.AppendFormat(" {0}", StreamName);
 
 			return SingleResponseCommandAsync<StreamRecord>(Arguments.ToString(), null, CancellationToken);
 		}
diff --git a/Engine/Source/Programs/Shared/EpicGames.Perforce/Records/StreamRecord.cs b/Engine/Source/Programs/Shared/EpicGames.Perforce/Records/StreamRecord.cs
--- a/Engine/Source/Programs/Shared/EpicGames.Perforce/Records/StreamRecord.cs
+++ b/Engine/Source/Programs/Shared/EpicGames.Perforce/Records/StreamRecord.cs
@@ -75,7 +75,7 @@ namespace EpicGames.Perforce
 		/// <summary>
 		/// Computed view for the stream
 		/// </summary>
-		[PerforceTag("View")]
+		[PerforceTag("View", Optional = true)]
 		public List<string> View = new List<string>();
 
 		/// <summary>
