P4VUtils: Add a command to merge a changelist to the current stream.

#rnx

diff --git a/Engine/Extras/P4VUtils/Commands/DescribeCommand.cs b/Engine/Extras/P4VUtils/Commands/DescribeCommand.cs
--- a/Engine/Extras/P4VUtils/Commands/DescribeCommand.cs
+++ b/Engine/Extras/P4VUtils/Commands/DescribeCommand.cs
@@ -16,7 +16,7 @@ namespace P4VUtils.Commands
 	{
 		public override string Description => "Prints out a description of the given changelist";
 
-		public override CustomToolInfo CustomTool => new CustomToolInfo("Describe", "%c");
+		public override CustomToolInfo CustomTool => new CustomToolInfo("Describe", "%c") { ShowConsole = true };
 
 		public override async Task<int> Execute(string[] Args, IReadOnlyDictionary<string, string> ConfigValues, ILogger Logger)
 		{
diff --git a/Engine/Extras/P4VUtils/Commands/EdigrateCommand.cs b/Engine/Extras/P4VUtils/Commands/EdigrateCommand.cs
--- a/Engine/Extras/P4VUtils/Commands/EdigrateCommand.cs
+++ b/Engine/Extras/P4VUtils/Commands/EdigrateCommand.cs
@@ -15,7 +15,7 @@ namespace P4VUtils.Commands
 	{
 		public override string Description => "Converts an integration to an edit";
 
-		public override CustomToolInfo CustomTool => new CustomToolInfo("Convert to Edigrate", "%p");
+		public override CustomToolInfo CustomTool => new CustomToolInfo("Convert to Edigrate", "%p") { ShowConsole = true };
 
 		public override async Task<int> Execute(string[] Args, IReadOnlyDictionary<string, string> ConfigValues, ILogger Logger)
 		{
diff --git a/Engine/Extras/P4VUtils/Commands/MergeCommand.cs b/Engine/Extras/P4VUtils/Commands/MergeCommand.cs
new file mode 100644
--- /dev/null
+++ b/Engine/Extras/P4VUtils/Commands/MergeCommand.cs
@@ -0,0 +1,61 @@
+using EpicGames.Perforce;
+using Microsoft.Extensions.Logging;
+using System;
+using System.Collections.Generic;
+using System.Linq;
+using System.Text;
+using System.Text.RegularExpressions;
+using System.Threading;
+using System.Threading.Tasks;
+
+namespace P4VUtils.Commands
+{
+	class MergeCommand : Command
+	{
+		public override string Description => "Merge a changelist into the current stream";
+
+		public override CustomToolInfo CustomTool => new CustomToolInfo("Merge to current stream", "%p") { ShowConsole = true };
+
+		public override async Task<int> Execute(string[] Args, IReadOnlyDictionary<string, string> ConfigValues, ILogger Logger)
+		{
+			int Change = int.Parse(Args[1]);
+
+			PerforceConnection Perforce = new PerforceConnection(null, null, null, Logger);
+
+			InfoRecord Info = await Perforce.GetInfoAsync(InfoOptions.None, CancellationToken.None);
+			if (Info.ClientStream == null)
+			{
+				Logger.LogError("Not currently in a stream workspace.");
+				return 1;
+			}
+
+			StreamRecord TargetStream = await Perforce.GetStreamAsync(Info.ClientStream, false, CancellationToken.None);
+			while (TargetStream.Type == "virtual" && TargetStream.Parent != null)
+			{
+				TargetStream = await Perforce.GetStreamAsync(TargetStream.Parent, false, CancellationToken.None);
+			}
+
+			DescribeRecord ExistingChangeRecord = await Perforce.DescribeAsync(Change, CancellationToken.None);
+			if (ExistingChangeRecord.Files.Count == 0)
+			{
+				Logger.LogError("No files in selected changelist");
+				return 1;
+			}
+
+			ChangeRecord NewChangeRecord = new ChangeRecord();
+			NewChangeRecord.User = Info.UserName;
+			NewChangeRecord.Client = Info.ClientName;
+			NewChangeRecord.Description = $"{ExistingChangeRecord.Description.TrimEnd()}\n#p4v-cherrypick {Change}";
+			NewChangeRecord = await Perforce.CreateChangeAsync(NewChangeRecord, CancellationToken.None);
+
+			string SourceFileSpec = Regex.Replace(ExistingChangeRecord.Files[0].DepotFile, @"^(//[^/]+/[^/]+)/.*$", "$1/...");
+			string TargetFileSpec = $"{TargetStream.Stream}/...";
+			await Perforce.MergeAsync(MergeOptions.None, NewChangeRecord.Number, -1, $"{SourceFileSpec}@={Change}", TargetFileSpec, CancellationToken.None);
+
+			Logger.LogInformation("Merged into pending changelist {0}", NewChangeRecord.Number);
+
+			await Perforce.TryResolveAsync(NewChangeRecord.Number, ResolveOptions.Automatic, null, CancellationToken.None);
+			return 0;
+		}
+	}
+}
diff --git a/Engine/Extras/P4VUtils/Commands/PreflightCommand.cs b/Engine/Extras/P4VUtils/Commands/PreflightCommand.cs
--- a/Engine/Extras/P4VUtils/Commands/PreflightCommand.cs
+++ b/Engine/Extras/P4VUtils/Commands/PreflightCommand.cs
@@ -21,7 +21,7 @@ namespace P4VUtils.Commands
 	{
 		public override string Description => "Runs a preflight of the given changelist on Horde";
 
-		public override CustomToolInfo CustomTool => new CustomToolInfo("Preflight...", "%p");
+		public override CustomToolInfo CustomTool => new CustomToolInfo("Horde: Preflight...", "%p");
 
 		public override async Task<int> Execute(string[] Args, IReadOnlyDictionary<string, string> ConfigValues, ILogger Logger)
 		{
@@ -30,8 +30,18 @@ namespace P4VUtils.Commands
 			PerforceConnection Perforce = new PerforceConnection(null, null, Logger);
 
 			ClientRecord Client = await Perforce.GetClientAsync(null, CancellationToken.None);
+			if(Client.Stream == null)
+			{
+				return 1;
+			}
+
+			StreamRecord Stream = await Perforce.GetStreamAsync(Client.Stream, false, CancellationToken.None);
+			while (Stream.Type == "virtual" && Stream.Parent != null)
+			{
+				Stream = await Perforce.GetStreamAsync(Stream.Parent, false, CancellationToken.None);
+			}
 
-			string Url = GetUrl(Client.Stream!, Change, ConfigValues);
+			string Url = GetUrl(Stream.Stream, Change, ConfigValues);
 			Logger.LogDebug("Opening {Url}", Url);
 			OpenUrl(Url);
 
@@ -69,7 +79,7 @@ namespace P4VUtils.Commands
 	{
 		public override string Description => "Runs a preflight of the given changelist on Horde and submits it";
 
-		public override CustomToolInfo CustomTool => new CustomToolInfo("Preflight and submit", "%p");
+		public override CustomToolInfo CustomTool => new CustomToolInfo("Horde: Preflight and submit", "%p");
 
 		public override string GetUrl(string Stream, int Change, IReadOnlyDictionary<string, string> ConfigValues)
 		{
diff --git a/Engine/Extras/P4VUtils/Program.cs b/Engine/Extras/P4VUtils/Program.cs
--- a/Engine/Extras/P4VUtils/Program.cs
+++ b/Engine/Extras/P4VUtils/Program.cs
@@ -44,6 +44,7 @@ namespace P4VUtils
 		{
 			["preflight"] = new PreflightCommand(),
 			["preflightandsubmit"] = new PreflightAndSubmitCommand(),
+			["merge"] = new MergeCommand(),
 			["edigrate"] = new EdigrateCommand(),
 			["describe"] = new DescribeCommand()
 		};
diff --git a/Engine/Extras/P4VUtils/Properties/launchSettings.json b/Engine/Extras/P4VUtils/Properties/launchSettings.json
--- a/Engine/Extras/P4VUtils/Properties/launchSettings.json
+++ b/Engine/Extras/P4VUtils/Properties/launchSettings.json
@@ -2,7 +2,7 @@
   "profiles": {
     "P4VUtils": {
       "commandName": "Project",
-      "commandLineArgs": "preflight 1234"
+      "commandLineArgs": "merge 15160963"
     }
   }
 }
\ No newline at end of file
diff --git a/Engine/Source/Programs/Shared/EpicGames.Perforce/Records/MergeOptions.cs b/Engine/Source/Programs/Shared/EpicGames.Perforce/Records/MergeOptions.cs
--- a/Engine/Source/Programs/Shared/EpicGames.Perforce/Records/MergeOptions.cs
+++ b/Engine/Source/Programs/Shared/EpicGames.Perforce/Records/MergeOptions.cs
@@ -12,6 +12,11 @@ namespace EpicGames.Perforce
 	[Flags]
 	public enum MergeOptions
 	{
+		/// <summary>
+		/// No options
+		/// </summary>
+		None = 0,
+
 		/// <summary>
 		/// Preview the merge
 		/// </summary>
