P4VUtils: Add cherry-pick, and convert-to-edit commands.

#rnx

diff --git a/Engine/Extras/P4VUtils/Commands/CherryPickCommand.cs b/Engine/Extras/P4VUtils/Commands/CherryPickCommand.cs
new file mode 100644
--- /dev/null
+++ b/Engine/Extras/P4VUtils/Commands/CherryPickCommand.cs
@@ -0,0 +1,75 @@
+using EpicGames.Perforce;
+using Microsoft.Extensions.Logging;
+using System;
+using System.Collections.Generic;
+using System.Linq;
+using System.Text;
+using System.Text.RegularExpressions;
+using System.Threading;
+using System.Threading.Tasks;
+
+namespace P4VUtils.Commands
+{
+	class CherryPickCommand : Command
+	{
+		public override string Description => "Cherry-picks a changelist into the current stream";
+
+		public override CustomToolInfo CustomTool => new CustomToolInfo("Cherry-Pick", "%S") { ShowConsole = true };
+
+		public override async Task<int> Execute(string[] Args, IReadOnlyDictionary<string, string> ConfigValues, ILogger Logger)
+		{
+			int Change = int.Parse(Args[1]);
+
+			PerforceConnection Perforce = new PerforceConnection(null, null, null, Logger);
+			int NewChange = await MergeAsync(Perforce, Change, Logger);
+			return (NewChange > 0) ? 0 : 1;
+		}
+
+		public static async Task<int> MergeAsync(PerforceConnection Perforce, int Change, ILogger Logger)
+		{
+			InfoRecord Info = await Perforce.GetInfoAsync(InfoOptions.None, CancellationToken.None);
+			if (Info.ClientStream == null)
+			{
+				Logger.LogError("Not currently in a stream workspace.");
+				return -1;
+			}
+
+			StreamRecord TargetStream = await Perforce.GetStreamAsync(Info.ClientStream, false, CancellationToken.None);
+			while (TargetStream.Type == "virtual" && TargetStream.Parent != null)
+			{
+				TargetStream = await Perforce.GetStreamAsync(TargetStream.Parent, false, CancellationToken.None);
+			}
+
+			DescribeRecord ExistingChangeRecord = await Perforce.DescribeAsync(Change, CancellationToken.None);
+			if (ExistingChangeRecord.Files.Count == 0)
+			{
+				Logger.LogError("No files in selected changelist");
+				return -1;
+			}
+
+			string SourceFileSpec = Regex.Replace(ExistingChangeRecord.Files[0].DepotFile, @"^(//[^/]+/[^/]+)/.*$", $"$1/...@={Change}");
+			string TargetFileSpec = $"{TargetStream.Stream}/...";
+			Logger.LogInformation("Merging from {SourceSpec} to {TargetSpec}", SourceFileSpec, TargetFileSpec);
+
+
+			ChangeRecord NewChangeRecord = new ChangeRecord();
+			NewChangeRecord.User = Info.UserName;
+			NewChangeRecord.Client = Info.ClientName;
+			NewChangeRecord.Description = $"{ExistingChangeRecord.Description.TrimEnd()}\n#p4v-cherrypick {Change}";
+			NewChangeRecord = await Perforce.CreateChangeAsync(NewChangeRecord, CancellationToken.None);
+			Logger.LogInformation("Created pending changelist {0}", NewChangeRecord.Number);
+
+			await Perforce.MergeAsync(MergeOptions.None, NewChangeRecord.Number, -1, SourceFileSpec, TargetFileSpec, CancellationToken.None);
+			Logger.LogInformation("Merged files into changelist {0}", SourceFileSpec, TargetFileSpec, NewChangeRecord.Number);
+
+			PerforceResponseList<ResolveRecord> ResolveRecords = await Perforce.TryResolveAsync(NewChangeRecord.Number, ResolveOptions.Automatic, null, CancellationToken.None);
+			if (!ResolveRecords.Succeeded)
+			{
+				Logger.LogError("Unable to resolve files. Please resolve manually.");
+				return -1;
+			}
+
+			return NewChangeRecord.Number;
+		}
+	}
+}
diff --git a/Engine/Extras/P4VUtils/Commands/ConvertToEditCommand.cs b/Engine/Extras/P4VUtils/Commands/ConvertToEditCommand.cs
new file mode 100644
--- /dev/null
+++ b/Engine/Extras/P4VUtils/Commands/ConvertToEditCommand.cs
@@ -0,0 +1,226 @@
+// Copyright Epic Games, Inc. All Rights Reserved.
+
+using EpicGames.Core;
+using EpicGames.Perforce;
+using Microsoft.Extensions.Logging;
+using System;
+using System.Collections.Generic;
+using System.ComponentModel;
+using System.Linq;
+using System.Text;
+using System.Threading;
+using System.Threading.Tasks;
+
+namespace P4VUtils.Commands
+{
+	class ConvertToEditCommand : Command
+	{
+		public override string Description => "Converts an integration to an edit";
+
+		public override CustomToolInfo CustomTool => new CustomToolInfo("Convert to Edit", "%p") { ShowConsole = true };
+
+		public override async Task<int> Execute(string[] Args, IReadOnlyDictionary<string, string> ConfigValues, ILogger Logger)
+		{
+			int Change = int.Parse(Args[1]);
+			bool Debug = Args.Any(x => x.Equals("-Debug", StringComparison.OrdinalIgnoreCase));
+
+			PerforceConnection Perforce = new PerforceConnection(null, null, Logger);
+			if(!await ConvertToEditAsync(Perforce, Change, Debug, Logger))
+			{
+				return 1;
+			}
+			return 0;
+		}
+
+		public static async Task<bool> ConvertToEditAsync(PerforceConnection Perforce, int Change, bool Debug, ILogger Logger)
+		{
+			// Get the list of opened files in the CL along with their metadata
+			Logger.LogInformation("Getting the list of files in changelist {Change}...", Change);
+			List<FStatRecord> OpenedRecords = await Perforce.GetOpenFilesAsync(OpenedOptions.None, Change, null, null, -1, null, CancellationToken.None);
+			if (OpenedRecords.Count == 0)
+			{
+				Logger.LogError("Change has no opened files. Aborting.");
+				return false;
+			}
+
+			// Each record has {depotFile, clientFile, rev, headRev, action, change, type, user, client}
+			Logger.LogInformation("There are {NumFiles} open files in the change.", OpenedRecords.Count);
+
+			// create a mapping of client to local files so we can re-open using local paths (needed for some actions)
+			List<WhereRecord> WhereRecords = await Perforce.WhereAsync(OpenedRecords.Select(x => x.DepotFile!).ToArray(), CancellationToken.None);
+			Dictionary<string, string> OpenedClientToLocalMap = WhereRecords.ToDictionary(x => x.ClientFile, x => x.Path, StringComparer.OrdinalIgnoreCase);
+
+			// If any file that needs to be resolved still, we need to abort.
+			// We don't want to convert non-resolved files to plain edits,
+			// as we need to resolve the merge first!
+			List<ResolveRecord> ResolveRecords = await Perforce.ResolveAsync(Change, ResolveOptions.PreviewOnly, null!, CancellationToken.None);
+			if (ResolveRecords.Any())
+			{
+				Logger.LogError("The following files in the changelist need to be resolved.");
+				Logger.LogError("This script cannot continue until all files are resolved:");
+				foreach (ResolveRecord ResolveRecord in ResolveRecords)
+				{
+					Logger.LogError("    {FileToResolve}", ResolveRecord.ClientFile);
+				}
+				return false;
+			}
+
+			// We'll just assume all the files in the changelist are integrated. We'll revert and re-open
+			// with the same action/filetype. This should remove the integration records.
+
+			// partition our files into ones we know the actions for, and the ones we don't
+			List<FStatRecord> DeleteActions = new List<FStatRecord>();
+			List<FStatRecord> AddActions = new List<FStatRecord>();
+			List<FStatRecord> EditActions = new List<FStatRecord>();
+			List<FStatRecord> MoveAddActions = new List<FStatRecord>();
+			List<FStatRecord> MoveDeleteActions = new List<FStatRecord>();
+			List<FStatRecord> UnknownActions = new List<FStatRecord>();
+
+			// Multiple action map to the same reopen mapping, so this handles that.
+			Dictionary<FileAction, List<FStatRecord>> ActionMappings = new Dictionary<FileAction, List<FStatRecord>>()
+			{
+				[FileAction.Delete] = DeleteActions,
+				[FileAction.Branch] = AddActions,
+				[FileAction.Add] = AddActions,
+				[FileAction.Edit] = EditActions,
+				[FileAction.Integrate] = EditActions,
+				[FileAction.MoveAdd] = MoveAddActions,
+				[FileAction.MoveDelete] = MoveDeleteActions
+			};
+
+			// this maps all open records to a reopen action.
+			foreach (FStatRecord OpenedRecord in OpenedRecords)
+			{
+				List<FStatRecord>? ActionList;
+				if (!ActionMappings.TryGetValue(OpenedRecord.Action, out ActionList))
+				{
+					ActionList = UnknownActions;
+				}
+				ActionList.Add(OpenedRecord);
+			}
+
+			//  if we have an unknown action, abort with error.
+			if (UnknownActions.Count > 0)
+			{
+				Logger.LogError("These files have unknown actions so cannot be reliably re-opened:");
+				foreach (FStatRecord OpenedRecord in UnknownActions)
+				{
+					Logger.LogError("  {Action}: {ClientFile}", OpenedRecord.Action.ToString(), OpenedRecord.ClientFile);
+				}
+				return false;
+			}
+
+			// We are good to go, so actually revert the files.
+			Logger.LogInformation("Reverting files that will be reopened");
+			if (!Debug)
+			{
+				PerforceResponseList<RevertRecord> Results = await Perforce.TryRevertAsync(-1, null, RevertOptions.KeepWorkspaceFiles, OpenedRecords.Select(x => x.ClientFile!).ToArray(), CancellationToken.None);
+				Results.EnsureSuccess();
+			}
+
+			// P4 lets you resolve/submit integrations without syncing them to the revision you are resolving against because it's a server operation.
+			// if the haveRev doesn't match the rev of the originally opened file, this was the case.
+			// so we sync -k because the local file matches the one we want, we just have to tell P4 that so it will let us check it out at that revision.
+			// If we simply sync -k to #head we might miss a legitimate submit that happened since our last resolve, and that would stomp that submit.
+			List<FStatRecord> NotSyncedOpenFiles = OpenedRecords.Where(x => x.HaveRevision != x.Revision && x.Action != FileAction.Branch && x.Action != FileAction.Add).ToList();
+			if (NotSyncedOpenFiles.Count > 0)
+			{
+				Logger.LogWarning("The following files were not synced to the revision they were resolved to before. syncing them to the indicated rev first.");
+				foreach (FStatRecord NotSyncedOpenFile in NotSyncedOpenFiles)
+				{
+					Logger.LogWarning("    {File}#{Revision} (had #{HaveRevision})", NotSyncedOpenFile.ClientFile, NotSyncedOpenFile.Revision, NotSyncedOpenFile.HaveRevision);
+				}
+				if (!Debug)
+				{
+					await Perforce.SyncAsync(SyncOptions.KeepWorkspaceFiles, -1, NotSyncedOpenFiles.Select(x => $"{x.ClientFile}#{x.Revision}").ToArray(), CancellationToken.None);
+				}
+			}
+
+			// Perform basic actions
+			await ReopenFiles("delete", DeleteActions, OpenedClientToLocalMap, x => Perforce.DeleteAsync(Change, DeleteOptions.None, x, CancellationToken.None), Debug, Logger);
+			await ReopenFiles("add", AddActions, OpenedClientToLocalMap, x => Perforce.AddAsync(Change, x, CancellationToken.None), Debug, Logger);
+			await ReopenFiles("edit", EditActions, OpenedClientToLocalMap, x => Perforce.EditAsync(Change, x, CancellationToken.None), Debug, Logger);
+
+			// Perform move actions
+			if (MoveAddActions.Count > 0)
+			{
+				Logger.LogInformation("Re-opening the following files for move:");
+
+				string[] LocalFiles = MoveAddActions.Select(x => OpenedClientToLocalMap[x.ClientFile!]).ToArray();
+				foreach (string LocalFile in LocalFiles)
+				{
+					Logger.LogInformation("    {LocalFile}", LocalFile);
+				}
+
+				// We have to first open the source file for edit (have to use -k because the file has already been moved locally!)
+				if (!Debug)
+				{
+					await Perforce.EditAsync(Change, null, EditOptions.KeepWorkspaceFiles, MoveAddActions.Select(x => x.MovedFile!).ToArray(), CancellationToken.None);
+
+					// then we can open the file for move in the new location (have to use -k because the file has already been moved locally!)
+					foreach ((FStatRecord OpenedRecord, string LocalFile) in MoveAddActions.Zip(LocalFiles))
+					{
+						await Perforce.MoveAsync(Change, null, MoveOptions.KeepWorkspaceFiles, OpenedRecord.MovedFile!, LocalFile, CancellationToken.None);
+					}
+				}
+			}
+
+			// Get the list of reopened files in the CL to check their filetype
+			Logger.LogInformation("Checking the list of files reopened in changelist {Change}...", Change);
+			List<FStatRecord> ReopenedRecords = await Perforce.GetOpenFilesAsync(OpenedOptions.None, Change, null, null, -1, null!, CancellationToken.None);
+			if (ReopenedRecords.Count == 0)
+			{
+				Logger.LogError("Change has no reopened files. This signifies an error in the script! Aborting...");
+				return false;
+			}
+			if (ReopenedRecords.Count != OpenedRecords.Count)
+			{
+				Logger.LogError("Change doesn't not have the same number of reopened files ({NumReopened}) as originally ({NumOpened}). This probably signifies an error in the script, and the actions should be reverted! Aborting...", ReopenedRecords.Count, OpenedRecords.Count);
+				return false;
+			}
+
+			// Ensure the filetypes match and ensure each reopened file was in the original changelist.
+			Dictionary<string, FStatRecord> ReopenedRecordsMap = ReopenedRecords.ToDictionary(x => x.ClientFile!, x => x);
+			foreach (FStatRecord OpenedRecord in OpenedRecords)
+			{
+				FStatRecord? ReopenedRecord;
+				if (!ReopenedRecordsMap.TryGetValue(OpenedRecord.ClientFile!, out ReopenedRecord))
+				{
+					Logger.LogError("Could not find original file {ClientFile} in re-opened records. This signifies an error in the script! Aborting...", OpenedRecord.ClientFile);
+					return false;
+				}
+				if (OpenedRecord.Type != ReopenedRecord.Type)
+				{
+					Logger.LogInformation("Changing filetype of {ClientFile} from {OpenType} to {ReopenType}", OpenedRecord.ClientFile!, ReopenedRecord.Type, OpenedRecord.Type);
+					if (!Debug)
+					{
+						await Perforce.ReopenAsync(Change, OpenedRecord.Type, ReopenedRecord.ClientFile!, CancellationToken.None);
+					}
+				}
+			}
+
+			// Done!
+			Logger.LogInformation("Success");
+			return true;
+		}
+
+		static async Task ReopenFiles(string Operation, List<FStatRecord> Records, Dictionary<string, string> OpenedClientToLocalMap, Func<string[], Task> ReopenAsync, bool Debug, ILogger Logger)
+		{
+			if (Records.Count > 0)
+			{
+				Logger.LogInformation("Re-opening the following files for {Operation}:", Operation);
+
+				string[] LocalFiles = Records.Select(x => OpenedClientToLocalMap[x.ClientFile!]).ToArray();
+				foreach (string LocalFile in LocalFiles)
+				{
+					Logger.LogInformation("    {LocalFile}", LocalFiles);
+				}
+
+				if (!Debug)
+				{
+					await ReopenAsync(LocalFiles);
+				}
+			}
+		}
+	}
+}
diff --git a/Engine/Extras/P4VUtils/Commands/EdigrateCommand.cs b/Engine/Extras/P4VUtils/Commands/EdigrateCommand.cs
--- a/Engine/Extras/P4VUtils/Commands/EdigrateCommand.cs
+++ b/Engine/Extras/P4VUtils/Commands/EdigrateCommand.cs
@@ -1,21 +1,21 @@
-// Copyright Epic Games, Inc. All Rights Reserved.
-
 using EpicGames.Perforce;
 using Microsoft.Extensions.Logging;
 using System;
 using System.Collections.Generic;
 using System.Linq;
 using System.Text;
+using System.Text.RegularExpressions;
 using System.Threading;
 using System.Threading.Tasks;
 
 namespace P4VUtils.Commands
 {
+
 	class EdigrateCommand : Command
 	{
-		public override string Description => "Converts an integration to an edit";
+		public override string Description => "Cherry-picks a change to the current stream as an edit";
 
-		public override CustomToolInfo CustomTool => new CustomToolInfo("Convert to Edigrate", "%p") { ShowConsole = true };
+		public override CustomToolInfo CustomTool => new CustomToolInfo("Edigrate", "%S") { ShowConsole = true };
 
 		public override async Task<int> Execute(string[] Args, IReadOnlyDictionary<string, string> ConfigValues, ILogger Logger)
 		{
@@ -24,193 +24,17 @@ namespace P4VUtils.Commands
 
 			PerforceConnection Perforce = new PerforceConnection(null, null, null, Logger);
 
-			// Get the list of opened files in the CL along with their metadata
-			Logger.LogInformation("Getting the list of files in changelist {Change}...", Change);
-			List<FStatRecord> OpenedRecords = await Perforce.GetOpenFilesAsync(OpenedOptions.None, Change, null, null, -1, null, CancellationToken.None);
-			if (OpenedRecords.Count == 0)
-			{
-				Logger.LogError("Change has no opened files. Aborting.");
-				return 1;
-			}
-
-			// Each record has {depotFile, clientFile, rev, headRev, action, change, type, user, client}
-			Logger.LogInformation("There are {NumFiles} open files in the change.", OpenedRecords.Count);
-
-			// create a mapping of client to local files so we can re-open using local paths (needed for some actions)
-			List<WhereRecord> WhereRecords = await Perforce.WhereAsync(OpenedRecords.Select(x => x.DepotFile!).ToArray(), CancellationToken.None);
-			Dictionary<string, string> OpenedClientToLocalMap = WhereRecords.ToDictionary(x => x.ClientFile, x => x.Path, StringComparer.OrdinalIgnoreCase);
-
-			// If any file that needs to be resolved still, we need to abort.
-			// We don't want to convert non-resolved files to plain edits,
-			// as we need to resolve the merge first!
-			List<ResolveRecord> ResolveRecords = await Perforce.ResolveAsync(Change, ResolveOptions.PreviewOnly, null!, CancellationToken.None);
-			if (ResolveRecords.Any())
-			{
-				Logger.LogError("The following files in the changelist need to be resolved.");
-				Logger.LogError("This script cannot continue until all files are resolved:");
-				foreach (ResolveRecord ResolveRecord in ResolveRecords)
-				{
-					Logger.LogError("    {FileToResolve}", ResolveRecord.ClientFile);
-				}
-				return 1;
-			}
-
-			// We'll just assume all the files in the changelist are integrated. We'll revert and re-open
-			// with the same action/filetype. This should remove the integration records.
-
-			// partition our files into ones we know the actions for, and the ones we don't
-			List<FStatRecord> DeleteActions = new List<FStatRecord>();
-			List<FStatRecord> AddActions = new List<FStatRecord>();
-			List<FStatRecord> EditActions = new List<FStatRecord>();
-			List<FStatRecord> MoveAddActions = new List<FStatRecord>();
-			List<FStatRecord> MoveDeleteActions = new List<FStatRecord>();
-			List<FStatRecord> UnknownActions = new List<FStatRecord>();
-
-			// Multiple action map to the same reopen mapping, so this handles that.
-			Dictionary<FileAction, List<FStatRecord>> ActionMappings = new Dictionary<FileAction, List<FStatRecord>>()
-			{
-				[FileAction.Delete] = DeleteActions,
-				[FileAction.Branch] = AddActions,
-				[FileAction.Add] = AddActions,
-				[FileAction.Edit] = EditActions,
-				[FileAction.Integrate] = EditActions,
-				[FileAction.MoveAdd] = MoveAddActions,
-				[FileAction.MoveDelete] = MoveDeleteActions
-			};
-
-			// this maps all open records to a reopen action.
-			foreach (FStatRecord OpenedRecord in OpenedRecords)
-			{
-				List<FStatRecord>? ActionList;
-				if (!ActionMappings.TryGetValue(OpenedRecord.Action, out ActionList))
-				{
-					ActionList = UnknownActions;
-				}
-				ActionList.Add(OpenedRecord);
-			}
-
-			//  if we have an unknown action, abort with error.
-			if (UnknownActions.Count > 0)
-			{
-				Logger.LogError("These files have unknown actions so cannot be reliably re-opened:");
-				foreach (FStatRecord OpenedRecord in UnknownActions)
-				{
-					Logger.LogError("  {Action}: {ClientFile}", OpenedRecord.Action.ToString(), OpenedRecord.ClientFile);
-				}
-				return 1;
-			}
-
-			// We are good to go, so actually revert the files.
-			Logger.LogInformation("Reverting files that will be reopened");
-			if (!Debug)
+			int MergedChange = await CherryPickCommand.MergeAsync(Perforce, Change, Logger);
+			if (MergedChange <= 0)
 			{
-				PerforceResponseList<RevertRecord> Results = await Perforce.TryRevertAsync(-1, null, RevertOptions.KeepWorkspaceFiles, OpenedRecords.Select(x => x.ClientFile!).ToArray(), CancellationToken.None);
-				Results.EnsureSuccess();
-			}
-
-			// P4 lets you resolve/submit integrations without syncing them to the revision you are resolving against because it's a server operation.
-			// if the haveRev doesn't match the rev of the originally opened file, this was the case.
-			// so we sync -k because the local file matches the one we want, we just have to tell P4 that so it will let us check it out at that revision.
-			// If we simply sync -k to #head we might miss a legitimate submit that happened since our last resolve, and that would stomp that submit.
-			List<FStatRecord> NotSyncedOpenFiles = OpenedRecords.Where(x => x.HaveRevision != x.Revision && x.Action != FileAction.Branch && x.Action != FileAction.Add).ToList();
-			if (NotSyncedOpenFiles.Count > 0)
-			{
-				Logger.LogWarning("The following files were not synced to the revision they were resolved to before. syncing them to the indicated rev first.");
-				foreach (FStatRecord NotSyncedOpenFile in NotSyncedOpenFiles)
-				{
-					Logger.LogWarning("    {File}#{Revision} (had #{HaveRevision})", NotSyncedOpenFile.ClientFile, NotSyncedOpenFile.Revision, NotSyncedOpenFile.HaveRevision);
-				}
-				if (!Debug)
-				{
-					await Perforce.SyncAsync(SyncOptions.KeepWorkspaceFiles, -1, NotSyncedOpenFiles.Select(x => $"{x.ClientFile}#{x.Revision}").ToArray(), CancellationToken.None);
-				}
-			}
-
-			// Perform basic actions
-			await ReopenFiles("delete", AddActions, OpenedClientToLocalMap, x => Perforce.DeleteAsync(Change, DeleteOptions.None, x, CancellationToken.None), Debug, Logger);
-			await ReopenFiles("add", AddActions, OpenedClientToLocalMap, x => Perforce.AddAsync(Change, x, CancellationToken.None), Debug, Logger);
-			await ReopenFiles("edit", AddActions, OpenedClientToLocalMap, x => Perforce.EditAsync(Change, x, CancellationToken.None), Debug, Logger);
-
-			// Perform move actions
-			if (MoveAddActions.Count > 0)
-			{
-				Logger.LogInformation("Re-opening the following files for move:");
-
-				string[] LocalFiles = MoveAddActions.Select(x => OpenedClientToLocalMap[x.ClientFile!]).ToArray();
-				foreach (string LocalFile in LocalFiles)
-				{
-					Logger.LogInformation("    {LocalFile}", LocalFile);
-				}
-
-				// We have to first open the source file for edit (have to use -k because the file has already been moved locally!)
-				if (!Debug)
-				{
-					await Perforce.EditAsync(Change, null, EditOptions.KeepWorkspaceFiles, MoveAddActions.Select(x => x.MovedFile!).ToArray(), CancellationToken.None);
-
-					// then we can open the file for move in the new location (have to use -k because the file has already been moved locally!)
-					foreach ((FStatRecord OpenedRecord, string LocalFile) in MoveAddActions.Zip(LocalFiles))
-					{
-						await Perforce.MoveAsync(Change, null, MoveOptions.KeepWorkspaceFiles, OpenedRecord.MovedFile!, LocalFile, CancellationToken.None);
-					}
-				}
-			}
-
-			// Get the list of reopened files in the CL to check their filetype
-			Logger.LogInformation("Getting the list of files reopened in changelist {Change}...", Change);
-			List<FStatRecord> ReopenedRecords = await Perforce.GetOpenFilesAsync(OpenedOptions.None, Change, null, null, -1, null!, CancellationToken.None);
-			if (ReopenedRecords.Count == 0)
-			{
-				Logger.LogError("Change has no reopened files. This signifies an error in the script! Aborting...");
 				return 1;
 			}
-			if (ReopenedRecords.Count != OpenedRecords.Count)
+			if (!await ConvertToEditCommand.ConvertToEditAsync(Perforce, MergedChange, Debug, Logger))
 			{
-				Logger.LogError("Change doesn't not have the same number of reopened files ({NumReopened}) as originally ({NumOpened}). This probably signifies an error in the script, and the actions should be reverted! Aborting...", ReopenedRecords.Count, OpenedRecords.Count);
 				return 1;
 			}
 
-			// Ensure the filetypes match and ensure each reopened file was in the original changelist.
-			Dictionary<string, FStatRecord> ReopenedRecordsMap = ReopenedRecords.ToDictionary(x => x.ClientFile!, x => x);
-			foreach (FStatRecord OpenedRecord in OpenedRecords)
-			{
-				FStatRecord? ReopenedRecord;
-				if (!ReopenedRecordsMap.TryGetValue(OpenedRecord.ClientFile!, out ReopenedRecord))
-				{
-					Logger.LogError("Could not find original file {ClientFile} in re-opened records. This signifies an error in the script! Aborting...", OpenedRecord.ClientFile);
-					return 1;
-				}
-				if (OpenedRecord.Type != ReopenedRecord.Type)
-				{
-					Logger.LogInformation("Changing filetype of {ClientFile} from {OpenType} to {ReopenType}", OpenedRecord.ClientFile!, ReopenedRecord.Type, OpenedRecord.Type);
-					if (!Debug)
-					{
-						await Perforce.ReopenAsync(Change, OpenedRecord.Type, ReopenedRecord.ClientFile!, CancellationToken.None);
-					}
-				}
-			}
-
-			// Done!
-			Logger.LogInformation("Finished");
 			return 0;
 		}
-
-		async Task ReopenFiles(string Operation, List<FStatRecord> Records, Dictionary<string, string> OpenedClientToLocalMap, Func<string[], Task> ReopenAsync, bool Debug, ILogger Logger)
-		{
-			if (Records.Count > 0)
-			{
-				Logger.LogInformation("Re-opening the following files for {Operation}:", Operation);
-
-				string[] LocalFiles = Records.Select(x => OpenedClientToLocalMap[x.ClientFile!]).ToArray();
-				foreach (string LocalFile in LocalFiles)
-				{
-					Logger.LogInformation("    {LocalFile}", LocalFiles);
-				}
-
-				if (!Debug)
-				{
-					await ReopenAsync(LocalFiles);
-				}
-			}
-		}
 	}
 }
diff --git a/Engine/Extras/P4VUtils/Commands/MergeCommand.cs b/Engine/Extras/P4VUtils/Commands/MergeCommand.cs
deleted file mode 100644
diff --git a/Engine/Extras/P4VUtils/Program.cs b/Engine/Extras/P4VUtils/Program.cs
--- a/Engine/Extras/P4VUtils/Program.cs
+++ b/Engine/Extras/P4VUtils/Program.cs
@@ -42,11 +42,12 @@ namespace P4VUtils
 	{
 		public static IReadOnlyDictionary<string, Command> Commands { get; } = new Dictionary<string, Command>
 		{
+			["converttoedit"] = new ConvertToEditCommand(),
+			["describe"] = new DescribeCommand(),
+			["cherrypick"] = new CherryPickCommand(),
+			["edigrate"] = new EdigrateCommand(),
 			["preflight"] = new PreflightCommand(),
 			["preflightandsubmit"] = new PreflightAndSubmitCommand(),
-			["merge"] = new MergeCommand(),
-			["edigrate"] = new EdigrateCommand(),
-			["describe"] = new DescribeCommand()
 		};
 
 		static void PrintHelp(ILogger Logger)
diff --git a/Engine/Source/Programs/Shared/EpicGames.Perforce/PerforceConnection.cs b/Engine/Source/Programs/Shared/EpicGames.Perforce/PerforceConnection.cs
--- a/Engine/Source/Programs/Shared/EpicGames.Perforce/PerforceConnection.cs
+++ b/Engine/Source/Programs/Shared/EpicGames.Perforce/PerforceConnection.cs
@@ -2176,7 +2176,7 @@ namespace EpicGames.Perforce
 		/// <param name="FileSpecs">Files to be reverted</param>
 		/// <param name="CancellationToken">Token used to cancel the operation</param>
 		/// <returns>Response from the server</returns>
-		public Task<PerforceResponseList<ResolveRecord>> TryResolveAsync(int ChangeNumber, ResolveOptions Options, string[]? FileSpecs, CancellationToken CancellationToken)
+		public async Task<PerforceResponseList<ResolveRecord>> TryResolveAsync(int ChangeNumber, ResolveOptions Options, string[]? FileSpecs, CancellationToken CancellationToken)
 		{
 			StringBuilder Arguments = new StringBuilder("resolve");
 			if ((Options & ResolveOptions.Automatic) != 0)
@@ -2227,7 +2227,9 @@ namespace EpicGames.Perforce
 				}
 			}
 
-			return CommandAsync<ResolveRecord>(Arguments.ToString(), null, CancellationToken);
+			PerforceResponseList<ResolveRecord> Records = await CommandAsync<ResolveRecord>(Arguments.ToString(), null, CancellationToken);
+			Records.RemoveAll(x => x.Error != null && x.Error.Generic == PerforceGenericCode.Empty);
+			return Records;
 		}
 
 		#endregion
diff --git a/Engine/Source/Programs/Shared/EpicGames.Perforce/PerforceResponseList.cs b/Engine/Source/Programs/Shared/EpicGames.Perforce/PerforceResponseList.cs
--- a/Engine/Source/Programs/Shared/EpicGames.Perforce/PerforceResponseList.cs
+++ b/Engine/Source/Programs/Shared/EpicGames.Perforce/PerforceResponseList.cs
@@ -46,7 +46,7 @@ namespace EpicGames.Perforce
 						return new List<T>();
 					}
 				}
-				return this.Select(x => x.Data).Where(x => x != null).ToList();
+				return this.Where(x => x.Info == null).Select(x => x.Data).Where(x => x != null).ToList();
 			}
 		}
 
diff --git a/Engine/Source/Programs/Shared/EpicGames.Perforce/Records/ResolveRecord.cs b/Engine/Source/Programs/Shared/EpicGames.Perforce/Records/ResolveRecord.cs
--- a/Engine/Source/Programs/Shared/EpicGames.Perforce/Records/ResolveRecord.cs
+++ b/Engine/Source/Programs/Shared/EpicGames.Perforce/Records/ResolveRecord.cs
@@ -16,32 +16,44 @@ namespace EpicGames.Perforce
 		/// <summary>
 		/// Path to the file in the workspace
 		/// </summary>
-		[PerforceTag("clientFile")]
-		public string ClientFile;
+		[PerforceTag("clientFile", Optional = true)]
+		public string? ClientFile;
 
 		/// <summary>
 		/// Path to the depot file that needs to be resolved
 		/// </summary>
-		[PerforceTag("fromFile")]
-		public string FromFile;
+		[PerforceTag("fromFile", Optional = true)]
+		public string? FromFile;
+
+		/// <summary>
+		/// Target file for the resolve
+		/// </summary>
+		[PerforceTag("toFile", Optional = true)]
+		public string? ToFile;
+
+		/// <summary>
+		/// How the file was resolved
+		/// </summary>
+		[PerforceTag("how", Optional = true)]
+		public string? How;
 
 		/// <summary>
 		/// Start range of changes to be resolved
 		/// </summary>
-		[PerforceTag("startFromRev")]
+		[PerforceTag("startFromRev", Optional = true)]
 		public int FromRevisionStart;
 
 		/// <summary>
 		/// Ending range of changes to be resolved
 		/// </summary>
-		[PerforceTag("endFromRev")]
+		[PerforceTag("endFromRev", Optional = true)]
 		public int FromRevisionEnd;
 
 		/// <summary>
 		/// The type of resolve to perform
 		/// </summary>
-		[PerforceTag("resolveType")]
-		public string ResolveType;
+		[PerforceTag("resolveType", Optional = true)]
+		public string? ResolveType;
 
 		/// <summary>
 		/// For content resolves, the type of resolve to be performed
@@ -63,7 +75,7 @@ namespace EpicGames.Perforce
 		/// Format this record for display in the debugger
 		/// </summary>
 		/// <returns>Summary of this revision</returns>
-		public override string ToString()
+		public override string? ToString()
 		{
 			return ClientFile;
 		}
