P4VUtils: Check the change is non-empty and shelve files before submit

#rnx

diff --git a/Engine/Extras/P4VUtils/Commands/PreflightCommand.cs b/Engine/Extras/P4VUtils/Commands/PreflightCommand.cs
--- a/Engine/Extras/P4VUtils/Commands/PreflightCommand.cs
+++ b/Engine/Extras/P4VUtils/Commands/PreflightCommand.cs
@@ -32,6 +32,16 @@ namespace P4VUtils.Commands
 			ClientRecord Client = await Perforce.GetClientAsync(null, CancellationToken.None);
 			if(Client.Stream == null)
 			{
+				Logger.LogError("Not being run from a stream client");
+				return 1;
+			}
+
+			await Perforce.ShelveAsync(Change, ShelveOptions.Overwrite, new[] { "//..." }, CancellationToken.None);
+
+			List<DescribeRecord> Describe = await Perforce.DescribeAsync(DescribeOptions.Shelved, -1, new[] { Change }, CancellationToken.None);
+			if (Describe[0].Files.Count == 0)
+			{
+				Logger.LogError("No files are shelved in the given changelist");
 				return 1;
 			}
 
diff --git a/Engine/Source/Programs/Shared/EpicGames.Perforce/PerforceChildProcess.cs b/Engine/Source/Programs/Shared/EpicGames.Perforce/PerforceChildProcess.cs
--- a/Engine/Source/Programs/Shared/EpicGames.Perforce/PerforceChildProcess.cs
+++ b/Engine/Source/Programs/Shared/EpicGames.Perforce/PerforceChildProcess.cs
@@ -389,7 +389,10 @@ namespace EpicGames.Perforce
 								break;
 							}
 
-							Responses.Add(Response);
+							if (Response.Error == null || Response.Error.Generic != PerforceGenericCode.Empty)
+							{
+								Responses.Add(Response);
+							}
 						}
 
 						// If the stream is complete but we couldn't parse a response from the server, treat it as an error
