P4VUtils: Add debug info to launching preflight.

#rnx

diff --git a/Engine/Extras/P4VUtils/Commands/PreflightCommand.cs b/Engine/Extras/P4VUtils/Commands/PreflightCommand.cs
--- a/Engine/Extras/P4VUtils/Commands/PreflightCommand.cs
+++ b/Engine/Extras/P4VUtils/Commands/PreflightCommand.cs
@@ -52,7 +52,7 @@ namespace P4VUtils.Commands
 			}
 
 			string Url = GetUrl(Stream.Stream, Change, ConfigValues);
-			Logger.LogDebug("Opening {Url}", Url);
+			Logger.LogInformation("Opening {Url}", Url);
 			OpenUrl(Url);
 
 			return 0;
