P4VUtils: Fix log output when converting multiple files to edit.

#rnx

diff --git a/Engine/Extras/P4VUtils/Commands/ConvertToEditCommand.cs b/Engine/Extras/P4VUtils/Commands/ConvertToEditCommand.cs
--- a/Engine/Extras/P4VUtils/Commands/ConvertToEditCommand.cs
+++ b/Engine/Extras/P4VUtils/Commands/ConvertToEditCommand.cs
@@ -224,7 +224,7 @@ namespace P4VUtils.Commands
 				string[] LocalFiles = Records.Select(x => OpenedClientToLocalMap[x.ClientFile!]).ToArray();
 				foreach (string LocalFile in LocalFiles)
 				{
-					Logger.LogInformation("    {LocalFile}", LocalFiles);
+					Logger.LogInformation("    {LocalFile}", LocalFile);
 				}
 
 				if (!Debug)
