P4VUtils: Add a compile command, which builds files in the selected changelist as non-unity. Automatically detects which project the files are in, etc...

diff --git a/Engine/Extras/P4VUtils/Commands/CompileCommand.cs b/Engine/Extras/P4VUtils/Commands/CompileCommand.cs
new file mode 100644
--- /dev/null
+++ b/Engine/Extras/P4VUtils/Commands/CompileCommand.cs
@@ -0,0 +1,192 @@
+// Copyright Epic Games, Inc. All Rights Reserved.
+
+using EpicGames.Core;
+using EpicGames.Perforce;
+using Microsoft.Extensions.Logging;
+using System;
+using System.Collections.Generic;
+using System.IO;
+using System.Linq;
+using System.Text;
+using System.Text.RegularExpressions;
+using System.Threading;
+using System.Threading.Tasks;
+
+namespace P4VUtils.Commands
+{
+	class CompileCommand : Command
+	{
+		public override string Description => "Compiles files in the selected changelist";
+
+		public override CustomToolInfo CustomTool => new CustomToolInfo("Compile Files (Non-Unity)", "%c") { ShowConsole = true };
+
+		public override async Task<int> Execute(string[] Args, IReadOnlyDictionary<string, string> ConfigValues, ILogger Logger)
+		{
+			int Change = -1;
+			if (Args.Length < 2)
+			{
+				Logger.LogError("Missing changelist number");
+				return 1;
+			}
+			else if (!Args[1].Equals("default", StringComparison.Ordinal) && !int.TryParse(Args[1], out Change))
+			{
+				Logger.LogError("'{Argument}' is not a numbered changelist", Args[1]);
+				return 1;
+			}
+
+			PerforceConnection Perforce = new PerforceConnection(null, null, null, Logger);
+			bool Result = await BuildAsync(Perforce, Change, Args.Skip(2), Logger);
+			return Result ? 0 : 1;
+		}
+
+		public static async Task<bool> BuildAsync(PerforceConnection Perforce, int Change, IEnumerable<string> AdditionalArguments, ILogger Logger)
+		{
+			ChangeRecord ChangeRecord = await Perforce.GetChangeAsync(GetChangeOptions.None, Change, CancellationToken.None);
+			if (ChangeRecord.Files.Count == 0)
+			{
+				Logger.LogError("No files in selected changelist");
+				return false;
+			}
+
+			List<FileReference> LocalFiles = new List<FileReference>();
+
+			List<WhereRecord> WhereRecords = await Perforce.WhereAsync(ChangeRecord.Files.ToArray(), CancellationToken.None);
+			foreach (WhereRecord WhereRecord in WhereRecords)
+			{
+				if(WhereRecord.Path != null)
+				{
+					FileReference LocalFile = new FileReference(WhereRecord.Path);
+					if (FileReference.Exists(LocalFile))
+					{
+						if (LocalFile.HasExtension(".cpp") || LocalFile.HasExtension(".h") || LocalFile.HasExtension(".c"))
+						{
+							LocalFiles.Add(LocalFile);
+						}
+					}
+				}
+			}
+
+			if (LocalFiles.Count == 0)
+			{
+				Logger.LogError("No files to build in this changelist");
+				return false;
+			}
+
+			HashSet<DirectoryReference> RootDirs = new HashSet<DirectoryReference>();
+			HashSet<FileReference> ProjectFiles = new HashSet<FileReference>();
+
+			HashSet<DirectoryReference> CheckedDirectories = new HashSet<DirectoryReference>();
+			foreach (FileReference LocalFile in LocalFiles)
+			{
+				DirectoryReference CurrentDir = LocalFile.Directory;
+				while (CheckedDirectories.Add(CurrentDir))
+				{
+					ProjectFiles.UnionWith(DirectoryReference.EnumerateFiles(CurrentDir, "*.uproject"));
+
+					FileReference BatchFile = FileReference.Combine(CurrentDir, "Engine/Build/BatchFiles/Build.bat");
+					if (FileReference.Exists(BatchFile))
+					{
+						RootDirs.Add(CurrentDir);
+						break;
+					}
+
+					DirectoryReference? NextDir = CurrentDir.ParentDirectory;
+					if (NextDir == null)
+					{
+						Logger.LogError("Unable to find engine directory for {FileName}", LocalFile);
+						return false;
+					}
+					CurrentDir = NextDir;
+				}
+			}
+
+			if (RootDirs.Count == 0)
+			{
+				Logger.LogError("Unable to find engine root directory");
+				return false;
+			}
+			if (RootDirs.Count > 1)
+			{
+				Logger.LogError("Found multiple engine root directories for files in changelist");
+				return false;
+			}
+
+			DirectoryReference RootDir = RootDirs.First();
+			DirectoryReference EngineDir = DirectoryReference.Combine(RootDir, "Engine");
+
+			List<string> Targets = new List<string>();
+
+			DirectoryReference EngineIntermediateDir = DirectoryReference.Combine(EngineDir, "Intermediate");
+			FileReference EngineFileList = FileReference.Combine(EngineIntermediateDir, "P4VUtils-Files.txt");
+			if (WriteFilteredFileList(EngineFileList, LocalFiles, EngineDir, Logger))
+			{
+				Targets.Add($"Win64 Development -TargetType=Editor -FileList={EngineFileList.FullName.QuoteArgument()}");
+			}
+			
+			foreach (FileReference ProjectFile in ProjectFiles)
+			{
+				DirectoryReference ProjectDir = ProjectFile.Directory;
+				FileReference ProjectFileList = FileReference.Combine(ProjectDir, "Intermediate", "P4VUtils-Files.txt");
+				if (WriteFilteredFileList(ProjectFileList, LocalFiles, ProjectDir, Logger))
+				{
+					Targets.Add($"Win64 Development -TargetType=Editor -FileList={ProjectFileList.FullName.QuoteArgument()} -Project={ProjectFile.FullName.QuoteArgument()}");
+				}
+			}
+
+			FileReference TargetListFile = FileReference.Combine(EngineIntermediateDir, "P4VUtils-Targets.txt");
+			WriteLines(TargetListFile, Targets);
+
+			FileReference BuildBatchFile = FileReference.Combine(RootDir, @"Engine\Build\BatchFiles\Build.bat");
+			StringBuilder Arguments = new StringBuilder($"{BuildBatchFile.FullName.QuoteArgument()} -TargetList={TargetListFile.FullName.QuoteArgument()}");
+			foreach (string AdditionalArgument in AdditionalArguments)
+			{
+				Arguments.AppendFormat(" {0}", AdditionalArgument.QuoteArgument());
+			}
+
+			Logger.LogInformation("Running {Arguments}", Arguments);
+			Logger.LogInformation("");
+
+			string ShellFileName = Environment.GetEnvironmentVariable("COMSPEC") ?? "C:\\Windows\\System32\\cmd.exe";
+			string ShellArguments = $"/C \"{Arguments}\"";
+
+			using (ManagedProcessGroup Group = new ManagedProcessGroup())
+			using (ManagedProcess Process = new ManagedProcess(Group, ShellFileName, ShellArguments, null, null, System.Diagnostics.ProcessPriorityClass.Normal))
+			{
+				await Process.CopyTo(Console.OpenStandardOutput(), CancellationToken.None);
+			}
+
+			return true;
+		}
+
+		static bool WriteFilteredFileList(FileReference FileList, List<FileReference> LocalFiles, DirectoryReference BaseDir, ILogger Logger)
+		{
+			List<FileReference> FilteredFiles = new List<FileReference>();
+			foreach (FileReference LocalFile in LocalFiles)
+			{
+				if (LocalFile.IsUnderDirectory(BaseDir))
+				{
+					FilteredFiles.Add(LocalFile);
+				}
+			}
+			if (FilteredFiles.Count > 0)
+			{
+				Logger.LogInformation("Files under {BaseDir}:", BaseDir);
+				foreach (FileReference FilteredFile in FilteredFiles)
+				{
+					Logger.LogInformation("  {File}", FilteredFile);
+				}
+				Logger.LogInformation("");
+
+				WriteLines(FileList, FilteredFiles.Select(x => x.FullName));
+				return true;
+			}
+			return false;
+		}
+
+		static void WriteLines(FileReference File, IEnumerable<string> Lines)
+		{
+			DirectoryReference.CreateDirectory(File.Directory);
+			FileReference.WriteAllLines(File, Lines);
+		}
+	}
+}
diff --git a/Engine/Extras/P4VUtils/Program.cs b/Engine/Extras/P4VUtils/Program.cs
--- a/Engine/Extras/P4VUtils/Program.cs
+++ b/Engine/Extras/P4VUtils/Program.cs
@@ -40,11 +40,12 @@ namespace P4VUtils
 
 	class Program
 	{
-		public static IReadOnlyDictionary<string, Command> Commands { get; } = new Dictionary<string, Command>
+		public static IReadOnlyDictionary<string, Command> Commands { get; } = new Dictionary<string, Command>(StringComparer.OrdinalIgnoreCase)
 		{
+			["cherrypick"] = new CherryPickCommand(),
+			["compile"] = new CompileCommand(),
 			["converttoedit"] = new ConvertToEditCommand(),
 			["describe"] = new DescribeCommand(),
-			["cherrypick"] = new CherryPickCommand(),
 			["edigrate"] = new EdigrateCommand(),
 			["preflight"] = new PreflightCommand(),
 			["preflightandsubmit"] = new PreflightAndSubmitCommand(),
diff --git a/Engine/Source/Programs/Shared/EpicGames.Perforce/PerforceConnection.cs b/Engine/Source/Programs/Shared/EpicGames.Perforce/PerforceConnection.cs
--- a/Engine/Source/Programs/Shared/EpicGames.Perforce/PerforceConnection.cs
+++ b/Engine/Source/Programs/Shared/EpicGames.Perforce/PerforceConnection.cs
@@ -554,7 +554,10 @@ namespace EpicGames.Perforce
 			{
 				Arguments.Append(" -O");
 			}
-			Arguments.AppendFormat(" {0}", ChangeNumber);
+			if (ChangeNumber != -1)
+			{
+				Arguments.AppendFormat(" {0}", ChangeNumber);
+			}
 
 			return SingleResponseCommandAsync<ChangeRecord>(Arguments.ToString(), null, CancellationToken);
 		}
