p4 custom tool for P4VUtils to find where the last edit occurred for a specific file

diff --git a/Engine/Extras/P4VUtils/Commands/CompileCommand.cs b/Engine/Extras/P4VUtils/Commands/CompileCommand.cs
--- a/Engine/Extras/P4VUtils/Commands/CompileCommand.cs
+++ b/Engine/Extras/P4VUtils/Commands/CompileCommand.cs
@@ -152,7 +152,7 @@ namespace P4VUtils.Commands
 			using (ManagedProcessGroup Group = new ManagedProcessGroup())
 			using (ManagedProcess Process = new ManagedProcess(Group, ShellFileName, ShellArguments, null, null, System.Diagnostics.ProcessPriorityClass.Normal))
 			{
-				await Process.CopyTo(Console.OpenStandardOutput(), CancellationToken.None);
+				await Process.CopyToAsync(Console.OpenStandardOutput(), CancellationToken.None);
 			}
 
 			return true;
diff --git a/Engine/Extras/P4VUtils/Commands/FindLastEditCommand.cs b/Engine/Extras/P4VUtils/Commands/FindLastEditCommand.cs
new file mode 100644
--- /dev/null
+++ b/Engine/Extras/P4VUtils/Commands/FindLastEditCommand.cs
@@ -0,0 +1,138 @@
+// Copyright Epic Games, Inc. All Rights Reserved.
+
+using EpicGames.Perforce;
+using Microsoft.Extensions.Logging;
+using System;
+using System.Collections.Generic;
+using System.Linq;
+using System.Threading;
+using System.Threading.Tasks;
+
+namespace P4VUtils.Commands
+{
+	class FindLastEditCommand : Command
+	{
+		public override string Description => "Prints who last edited this file and in which stream";
+
+		public override CustomToolInfo CustomTool => new CustomToolInfo("Find last edit", "%d") { ShowConsole = true };
+
+		private async Task<int> FindLastEditInFileHistory(PerforceConnection Perforce, string FilePath, ILogger Logger)
+		{
+			string[] Files = { FilePath };
+			List<FileLogRecord> FileLogRecords = await Perforce.FileLogAsync(FileLogOptions.None, Files, CancellationToken.None);
+
+			if (FileLogRecords.Count == 0)
+			{
+				Logger.LogWarning("Could not find any p4 change record for {0}", FilePath);
+				return 1;
+			}
+
+			// these are used to find out which integrate action are the source
+			IntegrateAction[] IntegrateFromActions =
+			{
+				IntegrateAction.BranchFrom,
+				IntegrateAction.MergeFrom,
+				IntegrateAction.MovedFrom,
+				IntegrateAction.CopyFrom,
+				IntegrateAction.DeleteFrom,
+				IntegrateAction.EditFrom,
+				IntegrateAction.AddFrom
+			};
+
+			// consider these 'edit' actions, stop searching
+			FileAction[] EditActions =
+			{
+				FileAction.Add,
+				FileAction.Edit,
+				FileAction.MoveAdd
+			};
+
+			// consider these 'integrate' actions, keep searching
+			FileAction[] IntegrateActions =
+			{
+				FileAction.Integrate,
+				FileAction.Branch,
+			};
+
+			// the first file log record is what we're interested in
+			// because we're not following integrations, it should be the only one anyway
+			FileLogRecord FileLogRecord = FileLogRecords.First();
+
+			// grab the last revision, check the action, continue unless it's an edit/add
+			foreach (RevisionRecord RevisionRecord in FileLogRecord.Revisions)
+			{
+				// skipping integrations that were ignored
+				if (RevisionRecord.Action == FileAction.Integrate && RevisionRecord.Integrations[0].Action == IntegrateAction.Ignored)
+				{
+					// go to the next revision record in this stream
+					continue;
+				}
+
+				// is it an 'integration'?
+				if (IntegrateActions.Contains(RevisionRecord.Action))
+				{
+					// where did we integrate from?
+					foreach (IntegrationRecord IntegrationRecord in RevisionRecord.Integrations)
+					{
+						if (IntegrateFromActions.Contains(IntegrationRecord.Action))
+						{
+							// we need to start the search at the integrated revision # otherwise we might loop forever. 
+							string NewFile = IntegrationRecord.OtherFile + "#" + IntegrationRecord.EndRevisionNumber;
+							return await FindLastEditInFileHistory(Perforce, NewFile, Logger);
+						}
+					}
+					Logger.LogError("Could not find integrate 'from' record");
+					break;
+				}
+				// is it an 'edit'?
+				else if (EditActions.Contains(RevisionRecord.Action))
+				{
+					// we found the perp so stop here
+
+					string[] SplitPath = FileLogRecord.DepotPath.Split("/", StringSplitOptions.RemoveEmptyEntries);
+					string Depot = SplitPath[0];
+					string Stream = SplitPath[1];
+
+					Logger.LogInformation("Last edit in //{0}/{1} by '{2}' in CL '{3}', revision #{4} ({5})",
+						Depot, Stream, RevisionRecord.UserName, RevisionRecord.ChangeNumber, RevisionRecord.RevisionNumber, RevisionRecord.Action.ToString());
+
+					return 0;
+				}
+			}
+
+			return 1;
+		}
+
+		public override async Task<int> Execute(string[] Args, IReadOnlyDictionary<string, string> ConfigValues, ILogger Logger)
+		{
+			if (Args.Length < 2)
+			{
+				Logger.LogError("Missing file name");
+				return 1;
+			}
+
+			string FileName = Args[1];
+
+			// %d is 'selected file or folder', picked that over %f so I don't have to run p4 where to 
+			// convert workspace path to depot path
+			if ( FileName.EndsWith("...") )
+			{
+				Logger.LogError("It looks like you selected a Folder, please select a File and run the tool again");
+				return 1;
+			}
+
+			Logger.LogInformation("Inspecting revision records for file '{0}'", FileName);
+
+			PerforceConnection Perforce = new PerforceConnection(null, null, Logger);
+
+			int Result = await FindLastEditInFileHistory(Perforce, FileName, Logger);
+
+			if ( Result != 0)
+			{
+				Logger.LogError("Unexpected error occurred.");
+			}
+				
+			return Result;
+		}
+	}
+}
diff --git a/Engine/Extras/P4VUtils/Program.cs b/Engine/Extras/P4VUtils/Program.cs
--- a/Engine/Extras/P4VUtils/Program.cs
+++ b/Engine/Extras/P4VUtils/Program.cs
@@ -49,6 +49,7 @@ namespace P4VUtils
 			["edigrate"] = new EdigrateCommand(),
 			["preflight"] = new PreflightCommand(),
 			["preflightandsubmit"] = new PreflightAndSubmitCommand(),
+			["findlastedit"] = new FindLastEditCommand(),
 		};
 
 		static void PrintHelp(ILogger Logger)
@@ -67,7 +68,7 @@ namespace P4VUtils
 			}
 
 			Logger.LogInformation("Commands:");
-			HelpUtils.PrintTable(Table, 2, 15);
+			HelpUtils.PrintTable(Table, 2, 15, Logger);
 		}
 
 		static async Task<int> Main(string[] Args)
@@ -115,7 +116,7 @@ namespace P4VUtils
 				if (Args.Any(x => x.Equals("-help", StringComparison.OrdinalIgnoreCase)))
 				{
 					List<KeyValuePair<string, string>> Parameters = CommandLineArguments.GetParameters(Command.GetType());
-					HelpUtils.PrintHelp(Args[0], Command.GetType());
+					HelpUtils.PrintHelp(Args[0], Command.GetType(), Logger);
 					return 0;
 				}
 
