P4VUtils: Use client name from P4CLIENT env var for preflights

Without it, P4 can fallback to inferring the client name from other places, causing a confusion of what name to use.
When P4V invokes a so called "custom tool", it will supply the necessary P4 configuration through the normal env vars.

diff --git a/Engine/Extras/P4VUtils/Commands/PreflightCommand.cs b/Engine/Extras/P4VUtils/Commands/PreflightCommand.cs
--- a/Engine/Extras/P4VUtils/Commands/PreflightCommand.cs
+++ b/Engine/Extras/P4VUtils/Commands/PreflightCommand.cs
@@ -37,9 +37,10 @@ namespace P4VUtils.Commands
 				return 1;
 			}
 
-			PerforceConnection Perforce = new PerforceConnection(null, null, Logger);
+			string? ClientName = Environment.GetEnvironmentVariable("P4CLIENT");
+			PerforceConnection Perforce = new PerforceConnection(null, null, ClientName, Logger);
 
-			ClientRecord Client = await Perforce.GetClientAsync(null, CancellationToken.None);
+			ClientRecord Client = await Perforce.GetClientAsync(ClientName, CancellationToken.None);
 			if(Client.Stream == null)
 			{
 				Logger.LogError("Not being run from a stream client");
