[P4VUtils] - add creation of a copy CL for preflightandSubmit to avoid change of ownership prolems that otherwise block horde from checking in the preflight CL

diff --git a/Engine/Extras/P4VUtils/Commands/PreflightCommand.cs b/Engine/Extras/P4VUtils/Commands/PreflightCommand.cs
--- a/Engine/Extras/P4VUtils/Commands/PreflightCommand.cs
+++ b/Engine/Extras/P4VUtils/Commands/PreflightCommand.cs
@@ -62,6 +62,40 @@ namespace P4VUtils.Commands
 				Stream = await Perforce.GetStreamAsync(Stream.Parent, false, CancellationToken.None);
 			}
 
+			if (CreateBackupCL())
+			{
+				// if this CL has files still open within it, create a new CL for those still opened files
+				// before sending the original CL to the preflight system - this avoids the problem where Horde
+				// cannot take ownership of the original CL and fails to checkin.
+				List<FStatRecord> OpenedRecords = await Perforce.GetOpenFilesAsync(OpenedOptions.None, Change, null, null, -1, null, CancellationToken.None);
+
+				if (OpenedRecords.Count > 0)
+				{
+					InfoRecord Info = await Perforce.GetInfoAsync(InfoOptions.None, CancellationToken.None);
+
+					ChangeRecord NewChangeRecord = new ChangeRecord();
+					NewChangeRecord.User = Info.UserName;
+					NewChangeRecord.Client = Info.ClientName;
+					NewChangeRecord.Description = $"{Describe[0].Description.TrimEnd()}\n#p4v-preflight-copy {Change}";
+					NewChangeRecord = await Perforce.CreateChangeAsync(NewChangeRecord, CancellationToken.None);
+
+					Logger.LogInformation("Created pending changelist {0}", NewChangeRecord.Number);
+
+					foreach (FStatRecord OpenedRecord in OpenedRecords)
+					{
+						if (OpenedRecord.ClientFile != null)
+						{
+							await Perforce.ReopenAsync(NewChangeRecord.Number, OpenedRecord.Type, OpenedRecord.ClientFile!, CancellationToken.None);
+							Logger.LogInformation("moving opened {file} to CL {CL}", OpenedRecord.ClientFile.ToString(), NewChangeRecord.Number);
+						}
+					}
+				}
+				else
+				{
+					Logger.LogInformation("No files opened, no copy CL created");
+				}
+			}
+
 			string Url = GetUrl(Stream.Stream, Change, ConfigValues);
 			Logger.LogInformation("Opening {Url}", Url);
 			OpenUrl(Url);
@@ -69,6 +103,11 @@ namespace P4VUtils.Commands
 			return 0;
 		}
 
+		public virtual bool CreateBackupCL()
+		{
+			return false;
+		}
+
 		public virtual string GetUrl(string Stream, int Change, IReadOnlyDictionary<string, string> ConfigValues)
 		{
 			string? BaseUrl;
@@ -106,5 +145,9 @@ namespace P4VUtils.Commands
 		{
 			return base.GetUrl(Stream, Change, ConfigValues) + "&defaulttemplate=true&submit=true";
 		}
+		public override bool CreateBackupCL()
+		{
+			return true;
+		}
 	}
 }
