[P4VUtils] - add snapshot & backout
backout command, performs an undo then convert to edit locally.
snapshot creates a shelved snapshot of a CL as a local temporary backup

diff --git a/Engine/Extras/P4VUtils/Commands/BackoutCommand.cs b/Engine/Extras/P4VUtils/Commands/BackoutCommand.cs
new file mode 100644
--- /dev/null
+++ b/Engine/Extras/P4VUtils/Commands/BackoutCommand.cs
@@ -0,0 +1,97 @@
+// Copyright Epic Games, Inc. All Rights Reserved.
+
+using EpicGames.Core;
+using EpicGames.Perforce;
+using Microsoft.Extensions.Logging;
+using System;
+using System.Collections.Generic;
+using System.ComponentModel;
+using System.Linq;
+using System.Text;
+using System.Threading;
+using System.Threading.Tasks;
+
+namespace P4VUtils.Commands
+{
+	class BackoutCommand : Command
+	{
+		public override string Description => "P4 Admin sanctioned method of backing out a CL";
+
+		public override CustomToolInfo CustomTool => new CustomToolInfo("Backout", "%p") { ShowConsole = true };
+
+		public override async Task<int> Execute(string[] Args, IReadOnlyDictionary<string, string> ConfigValues, ILogger Logger)
+		{
+			int Change;
+			if (Args.Length < 2)
+			{
+				Logger.LogError("Missing changelist number");
+				return 1;
+			}
+			else if (!int.TryParse(Args[1], out Change))
+			{
+				Logger.LogError("'{Argument}' is not a numbered changelist", Args[1]);
+				return 1;
+			}
+
+			bool Debug = Args.Any(x => x.Equals("-Debug", StringComparison.OrdinalIgnoreCase));
+
+			PerforceConnection Perforce = new PerforceConnection(null, null, Logger);
+
+			// Create a new CL
+			DescribeRecord ExistingChangeRecord = await Perforce.DescribeAsync(Change, CancellationToken.None);
+
+			InfoRecord Info = await Perforce.GetInfoAsync(InfoOptions.None, CancellationToken.None);
+
+			ChangeRecord NewChangeRecord = new ChangeRecord();
+			NewChangeRecord.User = Info.UserName;
+			NewChangeRecord.Client = Info.ClientName;
+			NewChangeRecord.Description = $"[Backout] - CL{Change}\n#fyi {ExistingChangeRecord.User}\nOriginal CL Desc\n-----------------------------------------------------------------\n{ExistingChangeRecord.Description.TrimEnd()}\n";
+			NewChangeRecord = await Perforce.CreateChangeAsync(NewChangeRecord, CancellationToken.None);
+			
+			Logger.LogInformation("Created pending changelist {0}", NewChangeRecord.Number);
+
+			// Undo the passed in CL
+			PerforceResponseList<UndoRecord> UndoResponses = await Perforce.TryUndoChangeAsync(Change, NewChangeRecord.Number, CancellationToken.None);
+
+			// Grab new CL info
+			DescribeRecord RefreshNewRecord = await Perforce.DescribeAsync(NewChangeRecord.Number, CancellationToken.None);
+
+			// if the original CL and the new CL differ in file count then an error occurs, abort and clean up
+			if (RefreshNewRecord.Files.Count != ExistingChangeRecord.Files.Count)
+			{
+				Logger.LogError("Undo on CL {0} failed", Change);
+				foreach (PerforceResponse Response in UndoResponses)
+				{
+					Logger.LogError("  {0}", Response.ToString())	;
+				}
+
+				// revert files in the new CL
+				if (RefreshNewRecord.Files.Count > 0)
+				{
+					Logger.LogError("  Reverting");
+					await Perforce.RevertAsync(RefreshNewRecord.Number, Info.ClientName, RevertOptions.None, new[] { "//..." }, CancellationToken.None);
+				}
+
+				// delete the new CL
+				Logger.LogError("  Deleting newly created CL {0}", NewChangeRecord.Number);
+				await Perforce.DeleteChangeAsync(DeleteChangeOptions.None, RefreshNewRecord.Number, CancellationToken.None);
+
+				return 1;
+			}
+			else
+			{
+				Logger.LogInformation("Undo of {0} created CL {1}", Change, NewChangeRecord.Number);
+			}
+
+			// Convert the undo CL over to an edit.
+			if(!await ConvertToEditCommand.ConvertToEditAsync(Perforce, NewChangeRecord.Number, Debug, Logger))
+			{
+				return 1;
+			}
+
+			return 0;
+		}
+
+	}
+
+}
diff --git a/Engine/Extras/P4VUtils/Commands/SnapShotCommand.cs b/Engine/Extras/P4VUtils/Commands/SnapShotCommand.cs
new file mode 100644
--- /dev/null
+++ b/Engine/Extras/P4VUtils/Commands/SnapShotCommand.cs
@@ -0,0 +1,101 @@
+// Copyright Epic Games, Inc. All Rights Reserved.
+
+using EpicGames.Core;
+using EpicGames.Perforce;
+using Microsoft.Extensions.Logging;
+using System;
+using System.Collections.Generic;
+using System.ComponentModel;
+using System.Diagnostics;
+using System.IO;
+using System.Reflection;
+using System.Runtime.InteropServices;
+using System.Text;
+using System.Text.RegularExpressions;
+using System.Threading;
+using System.Threading.Tasks;
+
+namespace P4VUtils.Commands
+{
+	class SnapshotCommand : Command
+	{
+		public override string Description => "Creates a shelved copy of the CL";
+
+		public override CustomToolInfo CustomTool => new CustomToolInfo("snapshot", "%p");
+
+		public override async Task<int> Execute(string[] Args, IReadOnlyDictionary<string, string> ConfigValues, ILogger Logger)
+		{
+			int Change;
+			if (Args.Length < 2)
+			{
+				Logger.LogError("Missing changelist number");
+				return 1;
+			}
+			else if (!int.TryParse(Args[1], out Change))
+			{
+				Logger.LogError("'{Argument}' is not a numbered changelist", Args[1]);
+				return 1;
+			}
+
+			string? ClientName = Environment.GetEnvironmentVariable("P4CLIENT");
+			PerforceConnection Perforce = new PerforceConnection(null, null, ClientName, Logger);
+
+			ClientRecord Client = await Perforce.GetClientAsync(ClientName, CancellationToken.None);
+			if(Client.Stream == null)
+			{
+				Logger.LogError("Not being run from a stream client");
+				return 1;
+			}
+
+			// Create the copy of the CL IFF it has files opened
+			List<FStatRecord> OpenedRecords = await Perforce.GetOpenFilesAsync(OpenedOptions.None, Change, null, null, -1, null, CancellationToken.None);
+
+			if (OpenedRecords.Count > 0)
+			{
+				InfoRecord Info = await Perforce.GetInfoAsync(InfoOptions.None, CancellationToken.None);
+				DescribeRecord ExistingChangeRecord = await Perforce.DescribeAsync(Change, CancellationToken.None);
+
+				var DateNow = DateTime.Now.ToString();
+
+				ChangeRecord NewChangeRecord = new ChangeRecord();
+				NewChangeRecord.User = Info.UserName;
+				NewChangeRecord.Client = Info.ClientName;
+				NewChangeRecord.Description = $"{ExistingChangeRecord.Description.TrimEnd()}\n[snapshot CL{Change} - {DateNow}]";
+				NewChangeRecord = await Perforce.CreateChangeAsync(NewChangeRecord, CancellationToken.None);
+
+				Logger.LogInformation("Created pending changelist {0}", NewChangeRecord.Number);
+
+				// Move the files to the new CL
+				foreach (FStatRecord OpenedRecord in OpenedRecords)
+				{
+					if (OpenedRecord.ClientFile != null)
+					{
+						await Perforce.ReopenAsync(NewChangeRecord.Number, OpenedRecord.Type, OpenedRecord.ClientFile!, CancellationToken.None);
+						Logger.LogInformation("moving opened {file} to CL {CL}", OpenedRecord.ClientFile.ToString(), NewChangeRecord.Number);
+					}
+				}
+
+				// Shelve that CL
+				await Perforce.ShelveAsync(NewChangeRecord.Number, ShelveOptions.Overwrite, new[] { "//..." }, CancellationToken.None);
+
+				// Move the files BACK to the old CL
+				foreach (FStatRecord OpenedRecord in OpenedRecords)
+				{
+					if (OpenedRecord.ClientFile != null)
+					{
+						await Perforce.ReopenAsync(Change, OpenedRecord.Type, OpenedRecord.ClientFile!, CancellationToken.None);
+						Logger.LogInformation("moving opened {file} to CL {CL}", OpenedRecord.ClientFile.ToString(), Change);
+					}
+				}
+			}
+			else
+			{
+				Logger.LogInformation("No files opened, snapshot not created");
+			}
+			
+
+			
+			return 0;
+		}
+	}
+}
diff --git a/Engine/Extras/P4VUtils/Program.cs b/Engine/Extras/P4VUtils/Program.cs
--- a/Engine/Extras/P4VUtils/Program.cs
+++ b/Engine/Extras/P4VUtils/Program.cs
@@ -50,6 +50,8 @@ namespace P4VUtils
 			["preflight"] = new PreflightCommand(),
 			["preflightandsubmit"] = new PreflightAndSubmitCommand(),
 			["findlastedit"] = new FindLastEditCommand(),
+			["snapshot"] = new SnapshotCommand(),
+			["backout"] = new BackoutCommand(),
 		};
 
 		static void PrintHelp(ILogger Logger)
@@ -124,7 +126,7 @@ namespace P4VUtils
 				return await Command.Execute(Args, ConfigValues, Logger);
 			}
 			else
-			{
+		{
 				Console.WriteLine("Unknown command: {0}", Args[0]);
 				PrintHelp(Logger);
 				return 1;
