P4VUtils: Fix warning.

diff --git a/Engine/Extras/P4VUtils/Program.cs b/Engine/Extras/P4VUtils/Program.cs
--- a/Engine/Extras/P4VUtils/Program.cs
+++ b/Engine/Extras/P4VUtils/Program.cs
@@ -210,7 +210,7 @@ namespace P4VUtils
 				return 1;
 			}
 
-			FileReference DotNetLocation = FileReference.Combine(DirectoryReference.GetSpecialFolder(Environment.SpecialFolder.ProgramFiles), "dotnet", "dotnet.exe");
+			FileReference DotNetLocation = FileReference.Combine(DirectoryReference.GetSpecialFolder(Environment.SpecialFolder.ProgramFiles)!, "dotnet", "dotnet.exe");
 			FileReference AssemblyLocation = new FileReference(Assembly.GetExecutingAssembly().GetOriginalLocation());
 
 			foreach (XmlNode? ChildNode in Root.SelectNodes("CustomToolDef"))
