[P4VUtils] - small feedback changes, addition of refresh and shortcut for all custom tools.
- snapshot now more descriptive
- backout now more descriptive
- preflightAndSubmit reverted to old behavior, new command to MoveWriteableFilesthenPreflightAndSubmit added

diff --git a/Engine/Extras/P4VUtils/Commands/BackoutCommand.cs b/Engine/Extras/P4VUtils/Commands/BackoutCommand.cs
--- a/Engine/Extras/P4VUtils/Commands/BackoutCommand.cs
+++ b/Engine/Extras/P4VUtils/Commands/BackoutCommand.cs
@@ -17,7 +17,7 @@ namespace P4VUtils.Commands
 	{
 		public override string Description => "P4 Admin sanctioned method of backing out a CL";
 
-		public override CustomToolInfo CustomTool => new CustomToolInfo("Backout", "%p") { ShowConsole = true };
+		public override CustomToolInfo CustomTool => new CustomToolInfo("Backout this CL", "%S") { ShowConsole = true };
 
 		public override async Task<int> Execute(string[] Args, IReadOnlyDictionary<string, string> ConfigValues, ILogger Logger)
 		{
diff --git a/Engine/Extras/P4VUtils/Commands/DescribeCommand.cs b/Engine/Extras/P4VUtils/Commands/DescribeCommand.cs
--- a/Engine/Extras/P4VUtils/Commands/DescribeCommand.cs
+++ b/Engine/Extras/P4VUtils/Commands/DescribeCommand.cs
@@ -16,7 +16,7 @@ namespace P4VUtils.Commands
 	{
 		public override string Description => "Prints out a description of the given changelist";
 
-		public override CustomToolInfo CustomTool => new CustomToolInfo("Describe", "%c") { ShowConsole = true };
+		public override CustomToolInfo CustomTool => new CustomToolInfo("Describe", "%c") { ShowConsole = true, RefreshUI = false};
 
 		public override async Task<int> Execute(string[] Args, IReadOnlyDictionary<string, string> ConfigValues, ILogger Logger)
 		{
diff --git a/Engine/Extras/P4VUtils/Commands/PreflightCommand.cs b/Engine/Extras/P4VUtils/Commands/PreflightCommand.cs
--- a/Engine/Extras/P4VUtils/Commands/PreflightCommand.cs
+++ b/Engine/Extras/P4VUtils/Commands/PreflightCommand.cs
@@ -145,6 +145,14 @@ namespace P4VUtils.Commands
 		{
 			return base.GetUrl(Stream, Change, ConfigValues) + "&defaulttemplate=true&submit=true";
 		}
+	}
+
+	class MoveWriteableFilesthenPreflightAndSubmitCommand : PreflightAndSubmitCommand
+	{
+		public override string Description => "Moves the writeable files to a new CL, then runs a preflight of the current changelist on Horde and submits it";
+
+		public override CustomToolInfo CustomTool => new CustomToolInfo("Horde: Move writeable files, Preflight and submit", "%p");
+
 		public override bool CreateBackupCL()
 		{
 			return true;
diff --git a/Engine/Extras/P4VUtils/Commands/SnapShotCommand.cs b/Engine/Extras/P4VUtils/Commands/SnapShotCommand.cs
--- a/Engine/Extras/P4VUtils/Commands/SnapShotCommand.cs
+++ b/Engine/Extras/P4VUtils/Commands/SnapShotCommand.cs
@@ -21,7 +21,7 @@ namespace P4VUtils.Commands
 	{
 		public override string Description => "Creates a shelved copy of the CL";
 
-		public override CustomToolInfo CustomTool => new CustomToolInfo("snapshot", "%p");
+		public override CustomToolInfo CustomTool => new CustomToolInfo("Create a Snapshot shelf of this CL", "%p");
 
 		public override async Task<int> Execute(string[] Args, IReadOnlyDictionary<string, string> ConfigValues, ILogger Logger)
 		{
diff --git a/Engine/Extras/P4VUtils/Program.cs b/Engine/Extras/P4VUtils/Program.cs
--- a/Engine/Extras/P4VUtils/Program.cs
+++ b/Engine/Extras/P4VUtils/Program.cs
@@ -21,6 +21,8 @@ namespace P4VUtils
 		public string Arguments { get; set; }
 		public bool AddToContextMenu { get; set; } = true;
 		public bool ShowConsole { get; set; }
+		public bool RefreshUI { get; set; } = true;
+		public string Shortcut { get; set; } = "";
 
 		public CustomToolInfo(string Name, string Arguments)
 		{
@@ -49,9 +51,11 @@ namespace P4VUtils
 			["edigrate"] = new EdigrateCommand(),
 			["preflight"] = new PreflightCommand(),
 			["preflightandsubmit"] = new PreflightAndSubmitCommand(),
+			["movewriteablepreflightandsubmit"] = new MoveWriteableFilesthenPreflightAndSubmitCommand(),
 			["findlastedit"] = new FindLastEditCommand(),
 			["snapshot"] = new SnapshotCommand(),
 			["backout"] = new BackoutCommand(),
+			["copyclnum"] = new CopyCLCommand(),
 		};
 
 		static void PrintHelp(ILogger Logger)
@@ -256,6 +260,13 @@ namespace P4VUtils
 							XmlElement Arguments = Document.CreateElement("Arguments");
 							Arguments.InnerText = $"{AssemblyLocation.FullName.QuoteArgument()} {Pair.Key} {CustomTool.Arguments}";
 							Definition.AppendChild(Arguments);
+
+							if (CustomTool.Shortcut.Length > 1)
+							{
+								XmlElement Shortcut = Document.CreateElement("Shortcut");
+								Shortcut.InnerText = CustomTool.Shortcut;
+								Definition.AppendChild(Shortcut);
+							}
 						}
 						ToolDef.AppendChild(Definition);
 
@@ -270,6 +281,13 @@ namespace P4VUtils
 							ToolDef.AppendChild(Console);
 						}
 
+						if (CustomTool.RefreshUI)
+						{
+							XmlElement Refresh = Document.CreateElement("Refresh");
+							Refresh.InnerText = CustomTool.RefreshUI ? "true" : "false";
+							ToolDef.AppendChild(Refresh);
+						}
+
 						XmlElement AddToContext = Document.CreateElement("AddToContext");
 						AddToContext.InnerText = CustomTool.AddToContextMenu ? "true" : "false";
 						ToolDef.AppendChild(AddToContext);
