[P4VUtils] - add CopyCL # command

diff --git a/Engine/Extras/P4VUtils/Commands/CopyCLCommand.cs b/Engine/Extras/P4VUtils/Commands/CopyCLCommand.cs
new file mode 100644
--- /dev/null
+++ b/Engine/Extras/P4VUtils/Commands/CopyCLCommand.cs
@@ -0,0 +1,59 @@
+// Copyright Epic Games, Inc. All Rights Reserved.
+
+using EpicGames.Core;
+using EpicGames.Perforce;
+using Microsoft.Extensions.Logging;
+using System;
+using System.Collections.Generic;
+using System.ComponentModel;
+using System.Text;
+using System.Threading;
+using System.Threading.Tasks;
+using System.Windows;
+
+namespace P4VUtils.Commands
+{
+	class CopyCLCommand : Command
+	{
+		public override string Description => "copy CL # only from pending and submitted CLs";
+
+		public override CustomToolInfo CustomTool => new CustomToolInfo("Copy CL #", "%c")
+		{
+			RefreshUI = false,
+			Shortcut = new string("Ctrl+Alt+C")
+		};
+
+		public override async Task<int> Execute(string[] Args, IReadOnlyDictionary<string, string> ConfigValues, ILogger Logger)
+		{
+			int Change;
+			if (Args.Length < 2)
+			{
+				Logger.LogError("Missing changelist number");
+				return 1;
+			}
+			else if (!int.TryParse(Args[1], out Change))
+			{
+				Logger.LogError("'{Argument}' is not a numbered changelist", Args[1]);
+				return 1;
+			}
+
+			// Jump through some hoops to make async & windows com happy and force STA on the clipboard call
+			var tcs = new TaskCompletionSource<int>();
+
+			var SetClipThread = new Thread(() => 
+				{ 
+					Clipboard.SetText(Change.ToString());
+					tcs.SetResult(0);
+				}
+			);
+
+			SetClipThread.SetApartmentState(ApartmentState.STA);
+			SetClipThread.Start();
+			SetClipThread.Join();
+
+			await tcs.Task;
+
+			return 0;
+		}
+	}
+}
diff --git a/Engine/Extras/P4VUtils/P4VUtils.csproj b/Engine/Extras/P4VUtils/P4VUtils.csproj
--- a/Engine/Extras/P4VUtils/P4VUtils.csproj
+++ b/Engine/Extras/P4VUtils/P4VUtils.csproj
@@ -1,9 +1,10 @@
-<Project Sdk="Microsoft.NET.Sdk">
+<Project Sdk="Microsoft.NET.Sdk.WindowsDesktop">
 
   <PropertyGroup>
     <OutputType>Exe</OutputType>
     <TargetFramework>netcoreapp3.1</TargetFramework>
     <Nullable>enable</Nullable>
+    <UseWPF>true</UseWPF>
   </PropertyGroup>
 
   <ItemGroup Condition="Exists('..\..\Restricted\NotForLicensees\Extras\P4VUtils\P4VUtils.ini')">
