[P4VUtils] - add folders for tools.

diff --git a/Engine/Extras/P4VUtils/Program.cs b/Engine/Extras/P4VUtils/Program.cs
--- a/Engine/Extras/P4VUtils/Program.cs
+++ b/Engine/Extras/P4VUtils/Program.cs
@@ -42,22 +42,35 @@
 
 	class Program
 	{
-		public static IReadOnlyDictionary<string, Command> Commands { get; } = new Dictionary<string, Command>(StringComparer.OrdinalIgnoreCase)
+		// UEHelpers - commands that help with common but simple operations
+		public static IReadOnlyDictionary<string, Command> HelperCommands { get; } = new Dictionary<string, Command>(StringComparer.OrdinalIgnoreCase)
 		{
+			["describe"] = new DescribeCommand(),
+			["findlastedit"] = new FindLastEditCommand(),
+			["snapshot"] = new SnapshotCommand(),
+			["copyclnum"] = new CopyCLCommand(),
+		};
+
+		// UEIntegrate Folder commands - complex commands to facilitate integrations/backout
+		public static IReadOnlyDictionary<string, Command> IntegrateCommands { get; } = new Dictionary<string, Command>(StringComparer.OrdinalIgnoreCase)
+		{
 			["cherrypick"] = new CherryPickCommand(),
-			["compile"] = new CompileCommand(),
 			["converttoedit"] = new ConvertToEditCommand(),
-			["describe"] = new DescribeCommand(),
 			["edigrate"] = new EdigrateCommand(),
+			["backout"] = new BackoutCommand(),
+		};
+
+		// UEHorde Folder - local build and horde preflights
+		public static IReadOnlyDictionary<string, Command> HordeCommands { get; } = new Dictionary<string, Command>(StringComparer.OrdinalIgnoreCase)
+		{
+			["compile"] = new CompileCommand(),
 			["preflight"] = new PreflightCommand(),
 			["preflightandsubmit"] = new PreflightAndSubmitCommand(),
 			["movewriteablepreflightandsubmit"] = new MoveWriteableFilesthenPreflightAndSubmitCommand(),
-			["findlastedit"] = new FindLastEditCommand(),
-			["snapshot"] = new SnapshotCommand(),
-			["backout"] = new BackoutCommand(),
-			["copyclnum"] = new CopyCLCommand(),
 		};
 
+		public static IDictionary<string, Command> Commands = HelperCommands.Concat(IntegrateCommands).Concat(HordeCommands).ToDictionary(p => p.Key, p => p.Value, StringComparer.OrdinalIgnoreCase);
+
 		static void PrintHelp(ILogger Logger)
 		{
 			Logger.LogInformation("P4VUtils");
@@ -182,46 +195,23 @@
 			return false;
 		}
 
-		public static async Task<int> UpdateCustomToolRegistration(bool bInstall, ILogger Logger)
+		static string GetToolName(XmlElement ToolNode)
 		{
-			DirectoryReference? ConfigDir = DirectoryReference.GetSpecialFolder(Environment.SpecialFolder.UserProfile);
-			if (ConfigDir == null)
+			return ToolNode.SelectSingleNode("Definition").SelectSingleNode("Name").InnerText;
+		}
+
+		// returns true if all tools were removed
+		static bool RemoveCustomToolsFromNode(XmlElement RootNode, FileReference DotNetLocation, FileReference AssemblyLocation, ILogger Logger)
+		{
+			int ToolsChecked = 0;
+			int ToolsRemoved = 0;
+			// Removes tools explicitly calling the assembly location identified above - i assume as a way to "filter" only those we explicitly added (@Ben.Marsh) - nochecking, remove this comment once verified.
+			foreach (XmlNode? ChildNode in RootNode.SelectNodes("CustomToolDef"))
 			{
-				Logger.LogError("Unable to find config directory.");
-				return 1;
-			}
-
-			FileReference ConfigFile = FileReference.Combine(ConfigDir, ".p4qt", "customtools.xml");
-
-			XmlDocument Document = new XmlDocument();
-			if (!TryLoadXmlDocument(ConfigFile, Document))
-			{
-				DirectoryReference.CreateDirectory(ConfigFile.Directory);
-				using (StreamWriter Writer = new StreamWriter(ConfigFile.FullName))
-				{
-					await Writer.WriteLineAsync(@"<?xml version=""1.0"" encoding=""UTF-8""?>");
-					await Writer.WriteLineAsync(@"<!--perforce-xml-version=1.0-->");
-					await Writer.WriteLineAsync(@"<CustomToolDefList varName=""customtooldeflist"">");
-					await Writer.WriteLineAsync(@"</CustomToolDefList>");
-				}
-				Document.Load(ConfigFile.FullName);
-			}
-
-			XmlElement? Root = Document.SelectSingleNode("CustomToolDefList") as XmlElement;
-			if (Root == null)
-			{
-				Logger.LogError("Unknown schema for {ConfigFile}", ConfigFile);
-				return 1;
-			}
-
-			FileReference DotNetLocation = FileReference.Combine(DirectoryReference.GetSpecialFolder(Environment.SpecialFolder.ProgramFiles)!, "dotnet", "dotnet.exe");
-			FileReference AssemblyLocation = new FileReference(Assembly.GetExecutingAssembly().GetOriginalLocation());
-
-			foreach (XmlNode? ChildNode in Root.SelectNodes("CustomToolDef"))
-			{
 				XmlElement? ChildElement = ChildNode as XmlElement;
 				if (ChildElement != null)
 				{
+					ToolsChecked++;
 					XmlElement? CommandElement = ChildElement.SelectSingleNode("Definition/Command") as XmlElement;
 					if (CommandElement != null && new FileReference(CommandElement.InnerText) == DotNetLocation)
 					{
@@ -231,17 +221,41 @@
 							string[] Arguments = CommandLineArguments.Split(ArgumentsElement.InnerText);
 							if (Arguments.Length > 0 && new FileReference(Arguments[0]) == AssemblyLocation)
 							{
-								Root.RemoveChild(ChildElement);
+								Logger.LogInformation("Removing Tool {0}", GetToolName(ChildElement));
+								RootNode.RemoveChild(ChildElement);
+								ToolsRemoved++;
 							}
 						}
 					}
 				}
 			}
 
-			// Insert new entries
-			if (bInstall)
+			return ToolsChecked == ToolsRemoved;
+		}
+
+		static void InstallCommandsListInFolder(string FolderName, IReadOnlyDictionary<string, Command> InputCommmands, XmlDocument Document, FileReference DotNetLocation, FileReference AssemblyLocation, ILogger Logger)
+		{
+			// <CustomToolDefList>				// list of custom tools (top level)
+			//  < CustomToolDef >				// loose custom tool in top level
+			//	< CustomToolFolder>				// folder containing custom tools
+			//		< Name > Test </ Name >
+			//		< CustomToolDefList >		// list of custom tools in folder
+			//		< CustomToolDef >			// definition of tool
+
+			// This is the top level node, there will also be a per folder node added of same name
+			XmlElement? Root = Document.SelectSingleNode("CustomToolDefList") as XmlElement;
+
+			if (Root != null)
 			{
-				foreach (KeyValuePair<string, Command> Pair in Commands)
+				XmlElement FolderDefinition = Document.CreateElement("CustomToolFolder");
+
+				XmlElement FolderDescription = Document.CreateElement("Name");
+				FolderDescription.InnerText = FolderName;
+				FolderDefinition.AppendChild(FolderDescription);
+
+				XmlElement FolderDefList = Document.CreateElement("CustomToolDefList");
+
+				foreach (KeyValuePair<string, Command> Pair in InputCommmands)
 				{
 					CustomToolInfo CustomTool = Pair.Value.CustomTool;
 
@@ -292,8 +306,90 @@
 						AddToContext.InnerText = CustomTool.AddToContextMenu ? "true" : "false";
 						ToolDef.AppendChild(AddToContext);
 					}
-					Root.AppendChild(ToolDef);
+					FolderDefList.AppendChild(ToolDef);
+				}
+
+				FolderDefinition.AppendChild(FolderDefList);
+
+				Root.AppendChild(FolderDefinition);
+			}
+		}
+		static void RemoveCustomToolsFromFolders(XmlElement RootNode, FileReference DotNetLocation, FileReference AssemblyLocation, ILogger Logger)
+		{
+			foreach (XmlNode? ChildNode in RootNode.SelectNodes("CustomToolFolder"))
+			{
+				if (ChildNode != null)
+				{
+					bool RemoveFolder = false;
+					XmlElement? FolderRoot = ChildNode.SelectSingleNode("CustomToolDefList") as XmlElement;
+					if (FolderRoot != null)
+					{
+						XmlElement? FolderNameNode = ChildNode.SelectSingleNode("Name") as XmlElement;
+						string FolderNameString = "";
+						if (FolderNameNode != null)
+						{
+							FolderNameString = FolderNameNode.InnerText;
+						}
+						Logger.LogInformation("Removing Tools from folder {0}", FolderNameString);
+						RemoveFolder = RemoveCustomToolsFromNode(FolderRoot, DotNetLocation, AssemblyLocation, Logger);
+					}
+
+					if (RemoveFolder)
+					{
+						// remove the folder itself.
+						RootNode.RemoveChild(ChildNode);
+					}
 				}
 			}
+		}
+		public static async Task<int> UpdateCustomToolRegistration(bool bInstall, ILogger Logger)
+		{
+			DirectoryReference? ConfigDir = DirectoryReference.GetSpecialFolder(Environment.SpecialFolder.UserProfile);
+			if (ConfigDir == null)
+			{
+				Logger.LogError("Unable to find config directory.");
+				return 1;
+			}
 
+			FileReference ConfigFile = FileReference.Combine(ConfigDir, ".p4qt", "customtools.xml");
+
+			XmlDocument Document = new XmlDocument();
+			if (!TryLoadXmlDocument(ConfigFile, Document))
+			{
+				DirectoryReference.CreateDirectory(ConfigFile.Directory);
+				using (StreamWriter Writer = new StreamWriter(ConfigFile.FullName))
+				{
+					await Writer.WriteLineAsync(@"<?xml version=""1.0"" encoding=""UTF-8""?>");
+					await Writer.WriteLineAsync(@"<!--perforce-xml-version=1.0-->");
+					await Writer.WriteLineAsync(@"<CustomToolDefList varName=""customtooldeflist"">");
+					await Writer.WriteLineAsync(@"</CustomToolDefList>");
+				}
+				Document.Load(ConfigFile.FullName);
+			}
+
+
+			FileReference DotNetLocation = FileReference.Combine(DirectoryReference.GetSpecialFolder(Environment.SpecialFolder.ProgramFiles)!, "dotnet", "dotnet.exe");
+			FileReference AssemblyLocation = new FileReference(Assembly.GetExecutingAssembly().GetOriginalLocation());
+
+			XmlElement? Root = Document.SelectSingleNode("CustomToolDefList") as XmlElement;
+			if (Root == null)
+			{
+				Logger.LogError("Unknown schema for {ConfigFile}", ConfigFile);
+				return 1;
+			}
+
+			// Remove Custom tools at the root
+			RemoveCustomToolsFromNode(Root, DotNetLocation, AssemblyLocation, Logger);
+
+			// Remove Custom tools in folders, and the folders
+			RemoveCustomToolsFromFolders(Root, DotNetLocation, AssemblyLocation, Logger);
+
+			// Insert new entries
+			if (bInstall)
+			{
+				InstallCommandsListInFolder("UEHelpers", HelperCommands, Document, DotNetLocation, AssemblyLocation, Logger);
+				InstallCommandsListInFolder("UEIntegrate", IntegrateCommands, Document, DotNetLocation, AssemblyLocation, Logger);
+				InstallCommandsListInFolder("UEHorde", HordeCommands, Document, DotNetLocation, AssemblyLocation, Logger);
+			}
+
 			// Save the new document
