Renaming p4vutil tool to make it more clear
#rnx

diff --git a/Engine/Extras/P4VUtils/Commands/P4BlameCommand.cs b/Engine/Extras/P4VUtils/Commands/P4BlameCommand.cs
--- a/Engine/Extras/P4VUtils/Commands/P4BlameCommand.cs
+++ b/Engine/Extras/P4VUtils/Commands/P4BlameCommand.cs
@@ -13,7 +13,7 @@ namespace P4VUtils.Commands
 	{
 		public override string Description => "Prints out who last edited a specific line in a file";
 
-		public override CustomToolInfo CustomTool => new CustomToolInfo("P4 Blame this file using line number", "%d") 
+		public override CustomToolInfo CustomTool => new CustomToolInfo("Find last edit by line", "%d") 
 		{ 
 			ShowConsole = true,
 			PromptForArgument = true,
@@ -36,7 +36,7 @@ namespace P4VUtils.Commands
 				return 1;
 			}
 
-			Logger.LogInformation("Blaming line {0} in {1}", LineNumber, FileName);
+			Logger.LogInformation("Finding last edit for line {0} in {1}", LineNumber, FileName);
 
 			PerforceConnection Perforce = new PerforceConnection(null, null, Logger);
 
diff --git a/Engine/Extras/P4VUtils/Program.cs b/Engine/Extras/P4VUtils/Program.cs
--- a/Engine/Extras/P4VUtils/Program.cs
+++ b/Engine/Extras/P4VUtils/Program.cs
@@ -49,7 +49,7 @@ namespace P4VUtils
 		{
 			["describe"] = new DescribeCommand(),
 			["findlastedit"] = new FindLastEditCommand(),
-			["p4blame"] = new P4BlameCommand(),
+			["findlasteditbyline"] = new P4BlameCommand(),
 			["snapshot"] = new SnapshotCommand(),
 			["copyclnum"] = new CopyCLCommand(),
 		};
