[P4VUtils] - preflightAndSubmit now warns the user if files are still checked out.

diff --git a/Engine/Extras/P4VUtils/Commands/PreflightCommand.cs b/Engine/Extras/P4VUtils/Commands/PreflightCommand.cs
--- a/Engine/Extras/P4VUtils/Commands/PreflightCommand.cs
+++ b/Engine/Extras/P4VUtils/Commands/PreflightCommand.cs
@@ -14,6 +14,7 @@ using System.Text;
 using System.Text.RegularExpressions;
 using System.Threading;
 using System.Threading.Tasks;
+using System.Windows;
 
 namespace P4VUtils.Commands
 {
@@ -95,6 +96,21 @@ namespace P4VUtils.Commands
 					Logger.LogInformation("No files opened, no copy CL created");
 				}
 			}
+			else
+			{
+				// if this CL has files still open within it, and this is a submit request - warn the user
+				List<FStatRecord> OpenedRecords = await Perforce.GetOpenFilesAsync(OpenedOptions.None, Change, null, null, -1, null, CancellationToken.None);
+
+				if (OpenedRecords.Count > 0 && IsSubmit())
+				{
+					if (MessageBox.Show("Do you want to continue?", "Changelist still has files checked out, your CL will fail to auto-submit", MessageBoxButton.YesNo) == MessageBoxResult.No)
+					{
+						Logger.LogInformation("User Opted to cancel");
+						return 0;
+					}
+				}
+			}
+
 
 			string Url = GetUrl(Stream.Stream, Change, ConfigValues);
 			Logger.LogInformation("Opening {Url}", Url);
@@ -118,6 +134,8 @@ namespace P4VUtils.Commands
 			return $"{BaseUrl.TrimEnd('/')}/preflight?stream={Stream}&change={Change}";
 		}
 
+		public virtual bool IsSubmit() { return false; }
+
 		void OpenUrl(string Url)
 		{
 			if (RuntimeInformation.IsOSPlatform(OSPlatform.Windows))
@@ -145,6 +163,8 @@ namespace P4VUtils.Commands
 		{
 			return base.GetUrl(Stream, Change, ConfigValues) + "&defaulttemplate=true&submit=true";
 		}
+
+		public override bool IsSubmit() { return true; }
 	}
 
 	class MoveWriteableFilesthenPreflightAndSubmitCommand : PreflightAndSubmitCommand
