[P4VUtils] - fix refresh post Snapshot command

diff --git a/Engine/Extras/P4VUtils/Commands/SnapShotCommand.cs b/Engine/Extras/P4VUtils/Commands/SnapShotCommand.cs
--- a/Engine/Extras/P4VUtils/Commands/SnapShotCommand.cs
+++ b/Engine/Extras/P4VUtils/Commands/SnapShotCommand.cs
@@ -21,7 +21,7 @@ namespace P4VUtils.Commands
 	{
 		public override string Description => "Creates a shelved copy of the CL";
 
-		public override CustomToolInfo CustomTool => new CustomToolInfo("Create a Snapshot shelf of this CL", "%p");
+		public override CustomToolInfo CustomTool => new CustomToolInfo("Create a Snapshot shelf of this CL", "%p") { ShowConsole = false, RefreshUI = true };
 
 		public override async Task<int> Execute(string[] Args, IReadOnlyDictionary<string, string> ConfigValues, ILogger Logger)
 		{
