[P4VUtils] - add Fast reconcile options for code + all - performs a reconcile on files in P4 that are writeable but NOT checked out

diff --git a/Engine/Extras/P4VUtils/Commands/FastReconcileEdits.cs b/Engine/Extras/P4VUtils/Commands/FastReconcileEdits.cs
new file mode 100644
--- /dev/null
+++ b/Engine/Extras/P4VUtils/Commands/FastReconcileEdits.cs
@@ -0,0 +1,123 @@
+// Copyright Epic Games, Inc. All Rights Reserved.
+
+using EpicGames.Core;
+using EpicGames.Perforce;
+using Microsoft.Extensions.Logging;
+using System;
+using System.Collections.Generic;
+using System.ComponentModel;
+using System.Diagnostics;
+using System.IO;
+using System.Reflection;
+using System.Runtime.InteropServices;
+using System.Text;
+using System.Text.RegularExpressions;
+using System.Threading;
+using System.Threading.Tasks;
+
+namespace P4VUtils.Commands
+{
+	class FastReconcileCodeEditsCommand : Command
+	{
+		
+		public override string Description => "Performs a fast reconcile of locally writeable code files into the default CL";
+
+		public override CustomToolInfo CustomTool => new CustomToolInfo("Fast Reconcile Writeable Code Files", "%D") { ShowConsole = true, RefreshUI = true };
+
+		public virtual string[] GetExtensions()
+		{
+			string[] Extensions =
+			{
+				".c*",
+				".h*",
+				".up*",
+				".ini"
+			};
+			return Extensions;
+		}
+
+		public override async Task<int> Execute(string[] Args, IReadOnlyDictionary<string, string> ConfigValues, ILogger Logger)
+		{
+			if (Args.Length < 2)
+			{
+				Logger.LogError("no folders selected");
+				return 1;
+			}
+
+			string[] Extensions = GetExtensions();
+
+			List<string> ReconcilePaths = new List<string>();
+
+			// %D is 'selected files or folders', ensure we only use folders
+			for (int i=1;i<Args.Length;i++)
+			{
+				Logger.LogInformation("params {i} - {Args[i]}", i, Args[i]);
+
+				if (Args[i].EndsWith("..."))
+				{
+					foreach (string ext in Extensions)
+					{
+						ReconcilePaths.Add(Args[i] + ext);
+					}
+				}
+				else
+				{
+					Logger.LogError("   ignoring individial selected file {file}", Args[i]);
+				}
+			}
+
+			foreach (string p in ReconcilePaths)
+			{
+				Logger.LogInformation("paths {p}", p);
+			}
+
+			string? ClientName = Environment.GetEnvironmentVariable("P4CLIENT");
+			PerforceConnection Perforce = new PerforceConnection(null, null, ClientName, Logger);
+
+			ClientRecord Client = await Perforce.GetClientAsync(ClientName, CancellationToken.None);
+			if(Client.Stream == null)
+			{
+				Logger.LogError("Not being run from a stream client");
+				return 1;
+			}
+
+			DateTime Start = DateTime.Now;
+
+
+			ReconcileOptions options = ReconcileOptions.Edit | ReconcileOptions.UseFileModification;
+
+			// Enable this to debug without taking action
+			// or potentially add the preview option for the user
+			//options |= ReconcileOptions.PreviewOnly;
+
+			List<ReconcileRecord> Results = await Perforce.ReconcileAsync(-1, options, ReconcilePaths.ToArray(), CancellationToken.None);
+
+			int Duration = Convert.ToInt32((DateTime.Now - Start).TotalSeconds);
+
+			foreach (ReconcileRecord res in Results)
+			{
+				Logger.LogInformation("Marked for Edit {res}", res);
+			}
+			
+			Logger.LogInformation("Completed in {Duration}s", Duration);
+
+			return 0;
+		}
+	}
+
+	class FastReconcileAllEditsCommand : FastReconcileCodeEditsCommand
+	{
+		public override string[] GetExtensions()
+		{
+			string[] Extensions =
+			{
+				".*",
+			};
+			return Extensions;
+		}
+
+		public override string Description => "Performs a fast reconcile of locally writeable files into the default CL";
+
+		public override CustomToolInfo CustomTool => new CustomToolInfo("Fast Reconcile Writeable Files", "%D") { ShowConsole = true, RefreshUI = true };
+	}
+}
diff --git a/Engine/Extras/P4VUtils/Program.cs b/Engine/Extras/P4VUtils/Program.cs
--- a/Engine/Extras/P4VUtils/Program.cs
+++ b/Engine/Extras/P4VUtils/Program.cs
@@ -52,6 +52,8 @@ namespace P4VUtils
 			["findlasteditbyline"] = new P4BlameCommand(),
 			["snapshot"] = new SnapshotCommand(),
 			["copyclnum"] = new CopyCLCommand(),
+			["reconcilecode"] = new FastReconcileCodeEditsCommand(),
+			["reconcileall"] = new FastReconcileAllEditsCommand(),
 		};
 
 		// UEIntegrate Folder commands - complex commands to facilitate integrations/backout
diff --git a/Engine/Source/Programs/Shared/EpicGames.Perforce/PerforceConnection.cs b/Engine/Source/Programs/Shared/EpicGames.Perforce/PerforceConnection.cs
--- a/Engine/Source/Programs/Shared/EpicGames.Perforce/PerforceConnection.cs
+++ b/Engine/Source/Programs/Shared/EpicGames.Perforce/PerforceConnection.cs
@@ -2082,7 +2082,11 @@ namespace EpicGames.Perforce
 			{
 				Arguments.Append(" -l");
 			}
-			foreach(string FileSpec in FileSpecs)
+			if ((Options & ReconcileOptions.UseFileModification) != 0)
+			{
+				Arguments.Append(" -m");
+			}
+			foreach (string FileSpec in FileSpecs)
 			{
 				Arguments.AppendFormat(" \"{0}\"", FileSpec);
 			}
diff --git a/Engine/Source/Programs/Shared/EpicGames.Perforce/Records/ReconcileOptions.cs b/Engine/Source/Programs/Shared/EpicGames.Perforce/Records/ReconcileOptions.cs
--- a/Engine/Source/Programs/Shared/EpicGames.Perforce/Records/ReconcileOptions.cs
+++ b/Engine/Source/Programs/Shared/EpicGames.Perforce/Records/ReconcileOptions.cs
@@ -53,5 +53,10 @@ namespace EpicGames.Perforce
 		/// Display output in local file syntax with relative paths, similar to the workspace-centric view of p4 status.
 		/// </summary>
 		LocalFileSyntax = 64,
+
+		/// <summary>
+		/// use file modification time before digest to optimize
+		/// </summary>
+		UseFileModification = 128,
 	}
 }
