[P4VUtils] - remove review triggers and stop preflight/snapshot from triggering secondary swarm reviews

diff --git a/Engine/Extras/P4VUtils/Commands/PreflightCommand.cs b/Engine/Extras/P4VUtils/Commands/PreflightCommand.cs
--- a/Engine/Extras/P4VUtils/Commands/PreflightCommand.cs
+++ b/Engine/Extras/P4VUtils/Commands/PreflightCommand.cs
@@ -20,6 +20,11 @@ namespace P4VUtils.Commands
 {
 	class PreflightCommand : Command
 	{
+
+		static public string StripReviewFyiHashTags(string InString)
+		{
+			return InString.Replace("#review", "-review").Replace("#codereview", "-codereview").Replace("#fyi", "-fyi");
+		}
 		public override string Description => "Runs a preflight of the given changelist on Horde";
 
 		public override CustomToolInfo CustomTool => new CustomToolInfo("Horde: Preflight...", "%p");
@@ -77,7 +82,7 @@ namespace P4VUtils.Commands
 					ChangeRecord NewChangeRecord = new ChangeRecord();
 					NewChangeRecord.User = Info.UserName;
 					NewChangeRecord.Client = Info.ClientName;
-					NewChangeRecord.Description = $"{Describe[0].Description.TrimEnd()}\n#p4v-preflight-copy {Change}";
+					NewChangeRecord.Description = $"{StripReviewFyiHashTags(Describe[0].Description.TrimEnd())}\n#p4v-preflight-copy {Change}";
 					NewChangeRecord = await Perforce.CreateChangeAsync(NewChangeRecord, CancellationToken.None);
 
 					Logger.LogInformation("Created pending changelist {0}", NewChangeRecord.Number);
diff --git a/Engine/Extras/P4VUtils/Commands/SnapShotCommand.cs b/Engine/Extras/P4VUtils/Commands/SnapShotCommand.cs
--- a/Engine/Extras/P4VUtils/Commands/SnapShotCommand.cs
+++ b/Engine/Extras/P4VUtils/Commands/SnapShotCommand.cs
@@ -60,7 +60,7 @@ namespace P4VUtils.Commands
 				ChangeRecord NewChangeRecord = new ChangeRecord();
 				NewChangeRecord.User = Info.UserName;
 				NewChangeRecord.Client = Info.ClientName;
-				NewChangeRecord.Description = $"{ExistingChangeRecord.Description.TrimEnd()}\n[snapshot CL{Change} - {DateNow}]";
+				NewChangeRecord.Description = $"{PreflightCommand.StripReviewFyiHashTags(ExistingChangeRecord.Description.TrimEnd())}\n[snapshot CL{Change} - {DateNow}]";
 				NewChangeRecord = await Perforce.CreateChangeAsync(NewChangeRecord, CancellationToken.None);
 
 				Logger.LogInformation("Created pending changelist {0}", NewChangeRecord.Number);
