[P4VUtils] - preflightandsubmit now provides options if the CL still has files within it.

diff --git a/Engine/Extras/P4VUtils/Commands/PreflightCommand.cs b/Engine/Extras/P4VUtils/Commands/PreflightCommand.cs
--- a/Engine/Extras/P4VUtils/Commands/PreflightCommand.cs
+++ b/Engine/Extras/P4VUtils/Commands/PreflightCommand.cs
@@ -15,6 +15,7 @@ using System.Text.RegularExpressions;
 using System.Threading;
 using System.Threading.Tasks;
 using System.Windows;
+using System.Linq;
 
 namespace P4VUtils.Commands
 {
@@ -68,12 +69,13 @@ namespace P4VUtils.Commands
 				Stream = await Perforce.GetStreamAsync(Stream.Parent, false, CancellationToken.None);
 			}
 
+			List<FStatRecord> OpenedRecords = await Perforce.GetOpenFilesAsync(OpenedOptions.None, Change, null, null, -1, null, CancellationToken.None);
+
 			if (CreateBackupCL())
 			{
 				// if this CL has files still open within it, create a new CL for those still opened files
 				// before sending the original CL to the preflight system - this avoids the problem where Horde
-				// cannot take ownership of the original CL and fails to checkin.
-				List<FStatRecord> OpenedRecords = await Perforce.GetOpenFilesAsync(OpenedOptions.None, Change, null, null, -1, null, CancellationToken.None);
+				// cannot take ownership of the original CL and fails to checkin
 
 				if (OpenedRecords.Count > 0)
 				{
@@ -104,15 +106,24 @@ namespace P4VUtils.Commands
 			else
 			{
 				// if this CL has files still open within it, and this is a submit request - warn the user
-				List<FStatRecord> OpenedRecords = await Perforce.GetOpenFilesAsync(OpenedOptions.None, Change, null, null, -1, null, CancellationToken.None);
 
 				if (OpenedRecords.Count > 0 && IsSubmit())
 				{
-					if (MessageBox.Show("Do you want to continue?", "Changelist still has files checked out, your CL will fail to auto-submit", MessageBoxButton.YesNo) == MessageBoxResult.No)
+					MessageBoxResult result= MessageBox.Show("Do you want to revert local files (yes) leave files as is (no) or cancel?", "Changelist still has files checked out, your CL will fail to auto-submit", MessageBoxButton.YesNoCancel);
+
+					if (result == MessageBoxResult.Cancel)
 					{
 						Logger.LogInformation("User Opted to cancel");
 						return 0;
 					}
+					else if(result == MessageBoxResult.Yes)
+					{
+						await Perforce.RevertAsync(Change, null, RevertOptions.None, OpenedRecords.Select(x => x.ClientFile!).ToArray(), CancellationToken.None);
+					}
+					else
+					{
+						// do nothing - user has been warned.
+					}
 				}
 			}
 
