[P4VUtils] - improving preflight warning message

diff --git a/Engine/Extras/P4VUtils/Commands/PreflightCommand.cs b/Engine/Extras/P4VUtils/Commands/PreflightCommand.cs
--- a/Engine/Extras/P4VUtils/Commands/PreflightCommand.cs
+++ b/Engine/Extras/P4VUtils/Commands/PreflightCommand.cs
@@ -105,11 +105,14 @@ namespace P4VUtils.Commands
 			}
 			else
 			{
-				// if this CL has files still open within it, and this is a submit request - warn the user
+				// if this CL has files still open within it, and this is a submit request - warn the user and provide options
 
 				if (OpenedRecords.Count > 0 && IsSubmit())
 				{
-					MessageBoxResult result= MessageBox.Show("Do you want to revert local files (yes) leave files as is (no) or cancel?", "Changelist still has files checked out, your CL will fail to auto-submit", MessageBoxButton.YesNoCancel);
+					MessageBoxResult result= MessageBox.Show(
+						"Do you want to revert local files and submit (yes), start the preflight anyways (which will fail to submit on finish if files are still checked out) (no), or cancel", 
+						"Changelist still has files checked out, your CL will fail to auto-submit", 
+						MessageBoxButton.YesNoCancel);
 
 					if (result == MessageBoxResult.Cancel)
 					{
