[P4VUtils] - improve messaging around preflight command when files remain checked out

diff --git a/Engine/Extras/P4VUtils/Commands/PreflightCommand.cs b/Engine/Extras/P4VUtils/Commands/PreflightCommand.cs
--- a/Engine/Extras/P4VUtils/Commands/PreflightCommand.cs
+++ b/Engine/Extras/P4VUtils/Commands/PreflightCommand.cs
@@ -54,8 +54,14 @@ namespace P4VUtils.Commands
 				return 1;
 			}
 
-			await Perforce.ShelveAsync(Change, ShelveOptions.Overwrite, new[] { "//..." }, CancellationToken.None);
+			List<FStatRecord> OpenedRecords = await Perforce.GetOpenFilesAsync(OpenedOptions.None, Change, null, null, -1, null, CancellationToken.None);
 
+			if (OpenedRecords.Count > 0)
+			{
+				Logger.LogInformation("Shelving changelist {0}", Change);
+				await Perforce.ShelveAsync(Change, ShelveOptions.Overwrite, new[] { "//..." }, CancellationToken.None);
+			}
+				
 			List<DescribeRecord> Describe = await Perforce.DescribeAsync(DescribeOptions.Shelved, -1, new[] { Change }, CancellationToken.None);
 			if (Describe[0].Files.Count == 0)
 			{
@@ -69,8 +75,6 @@ namespace P4VUtils.Commands
 				Stream = await Perforce.GetStreamAsync(Stream.Parent, false, CancellationToken.None);
 			}
 
-			List<FStatRecord> OpenedRecords = await Perforce.GetOpenFilesAsync(OpenedOptions.None, Change, null, null, -1, null, CancellationToken.None);
-
 			if (CreateBackupCL())
 			{
 				// if this CL has files still open within it, create a new CL for those still opened files
@@ -110,8 +114,15 @@ namespace P4VUtils.Commands
 				if (OpenedRecords.Count > 0 && IsSubmit())
 				{
 					MessageBoxResult result= MessageBox.Show(
-						"Do you want to revert local files and submit (yes), start the preflight anyways (which will fail to submit on finish if files are still checked out) (no), or cancel", 
-						"Changelist still has files checked out, your CL will fail to auto-submit", 
+						"Your CL was just shelved however it still has files checked out\r\n" +
+						"If the files remain in the CL your preflight will fail to submit\r\n" +
+						"\r\n" +
+						"Click:\r\n" +
+						"[YES] - To revert local files and submit the preflight,\r\n" +
+						"[NO] - To start the preflight, and move the files manually\r\n" +
+						"[CANCEL] - To cancel the request", 
+
+						"Your CL will fail to auto-submit unless fixed", 
 						MessageBoxButton.YesNoCancel);
 
 					if (result == MessageBoxResult.Cancel)
