[P4VUtils] - fix up FileSpecList usage when null is being passed in.

diff --git a/Engine/Extras/P4VUtils/Commands/CherryPickCommand.cs b/Engine/Extras/P4VUtils/Commands/CherryPickCommand.cs
--- a/Engine/Extras/P4VUtils/Commands/CherryPickCommand.cs
+++ b/Engine/Extras/P4VUtils/Commands/CherryPickCommand.cs
@@ -74,7 +74,7 @@ namespace P4VUtils.Commands
 			await Perforce.MergeAsync(MergeOptions.None, NewChangeRecord.Number, -1, SourceFileSpec, TargetFileSpec, CancellationToken.None);
 			Logger.LogInformation("Merged files into changelist {0}", SourceFileSpec, TargetFileSpec, NewChangeRecord.Number);
 
-			PerforceResponseList<ResolveRecord> ResolveRecords = await Perforce.TryResolveAsync(NewChangeRecord.Number, ResolveOptions.Automatic, null, CancellationToken.None);
+			PerforceResponseList<ResolveRecord> ResolveRecords = await Perforce.TryResolveAsync(NewChangeRecord.Number, ResolveOptions.Automatic, FileSpecList.Empty, CancellationToken.None);
 			if (!ResolveRecords.Succeeded)
 			{
 				Logger.LogError("Unable to resolve files. Please resolve manually.");
diff --git a/Engine/Extras/P4VUtils/Commands/ConvertToEditCommand.cs b/Engine/Extras/P4VUtils/Commands/ConvertToEditCommand.cs
--- a/Engine/Extras/P4VUtils/Commands/ConvertToEditCommand.cs
+++ b/Engine/Extras/P4VUtils/Commands/ConvertToEditCommand.cs
@@ -47,7 +47,7 @@ namespace P4VUtils.Commands
 		{
 			// Get the list of opened files in the CL along with their metadata
 			Logger.LogInformation("Getting the list of files in changelist {Change}...", Change);
-			List<FStatRecord> OpenedRecords = await Perforce.GetOpenFilesAsync(OpenedOptions.None, Change, null, null, -1, null, CancellationToken.None);
+			List<FStatRecord> OpenedRecords = await Perforce.GetOpenFilesAsync(OpenedOptions.None, Change, null, null, -1, FileSpecList.Empty, CancellationToken.None);
 			if (OpenedRecords.Count == 0)
 			{
 				Logger.LogError("Change has no opened files. Aborting.");
@@ -64,7 +64,7 @@ namespace P4VUtils.Commands
 			// If any file that needs to be resolved still, we need to abort.
 			// We don't want to convert non-resolved files to plain edits,
 			// as we need to resolve the merge first!
-			List<ResolveRecord> ResolveRecords = await Perforce.ResolveAsync(Change, ResolveOptions.PreviewOnly, null!, CancellationToken.None);
+			List<ResolveRecord> ResolveRecords = await Perforce.ResolveAsync(Change, ResolveOptions.PreviewOnly, FileSpecList.Empty!, CancellationToken.None);
 			if (ResolveRecords.Any())
 			{
 				Logger.LogError("The following files in the changelist need to be resolved.");
@@ -178,7 +178,7 @@ namespace P4VUtils.Commands
 
 			// Get the list of reopened files in the CL to check their filetype
 			Logger.LogInformation("Checking the list of files reopened in changelist {Change}...", Change);
-			List<FStatRecord> ReopenedRecords = await Perforce.GetOpenFilesAsync(OpenedOptions.None, Change, null, null, -1, null!, CancellationToken.None);
+			List<FStatRecord> ReopenedRecords = await Perforce.GetOpenFilesAsync(OpenedOptions.None, Change, null, null, -1, FileSpecList.Empty!, CancellationToken.None);
 			if (ReopenedRecords.Count == 0)
 			{
 				Logger.LogError("Change has no reopened files. This signifies an error in the script! Aborting...");
diff --git a/Engine/Extras/P4VUtils/Commands/PreflightCommand.cs b/Engine/Extras/P4VUtils/Commands/PreflightCommand.cs
--- a/Engine/Extras/P4VUtils/Commands/PreflightCommand.cs
+++ b/Engine/Extras/P4VUtils/Commands/PreflightCommand.cs
@@ -54,7 +54,7 @@ namespace P4VUtils.Commands
 				return 1;
 			}
 
-			List<FStatRecord> OpenedRecords = await Perforce.GetOpenFilesAsync(OpenedOptions.None, Change, null, null, -1, null, CancellationToken.None);
+			List<FStatRecord> OpenedRecords = await Perforce.GetOpenFilesAsync(OpenedOptions.None, Change, null, null, -1, FileSpecList.Empty, CancellationToken.None);
 
 			if (OpenedRecords.Count > 0)
 			{
diff --git a/Engine/Extras/P4VUtils/Commands/SnapShotCommand.cs b/Engine/Extras/P4VUtils/Commands/SnapShotCommand.cs
--- a/Engine/Extras/P4VUtils/Commands/SnapShotCommand.cs
+++ b/Engine/Extras/P4VUtils/Commands/SnapShotCommand.cs
@@ -48,7 +48,7 @@ namespace P4VUtils.Commands
 			}
 
 			// Create the copy of the CL IFF it has files opened
-			List<FStatRecord> OpenedRecords = await Perforce.GetOpenFilesAsync(OpenedOptions.None, Change, null, null, -1, null, CancellationToken.None);
+			List<FStatRecord> OpenedRecords = await Perforce.GetOpenFilesAsync(OpenedOptions.None, Change, null, null, -1, FileSpecList.Empty, CancellationToken.None);
 
 			if (OpenedRecords.Count > 0)
 			{
