New perforce tools
Data tools:
* Convert data to local writable -> Convert All data in changelist or all selected data files to local writable

Unshelving tools:
* Unshelve to current revision -> Remembers revision all files you are about to unshelve and if the revision is older, will sync to saved revision
* Unshelve and make data writable -> Unshelve changelist but for data will make them writable locally

diff --git a/Engine/Extras/P4VUtils/Commands/DataCommands.cs b/Engine/Extras/P4VUtils/Commands/DataCommands.cs
new file mode 100644
--- /dev/null
+++ b/Engine/Extras/P4VUtils/Commands/DataCommands.cs
@@ -0,0 +1,81 @@
+// Copyright Epic Games, Inc. All Rights Reserved.
+
+using EpicGames.Core;
+using EpicGames.Perforce;
+using Microsoft.Extensions.Logging;
+using System;
+using System.Collections.Generic;
+using System.ComponentModel;
+using System.IO;
+using System.Linq;
+using System.Text;
+using System.Threading;
+using System.Threading.Tasks;
+using System.Windows;
+
+namespace P4VUtils.Commands
+{
+	abstract class DataCommandBase : Command
+	{
+
+		protected async Task<int> MakeDataFilesLocalWritable(List<(string, string)> FilesInfo, ILogger Logger)
+		{
+			PerforceConnection Perforce = new PerforceConnection(null, null, Logger);
+
+			List<string> FilesToMakeLocalWritable = FilesInfo.Where(t => t.Item2.StartsWith("binary")).Select(t => t.Item1).ToList();
+			List<RevertRecord> Results = await Perforce.RevertAsync(-1, null, RevertOptions.KeepWorkspaceFiles, FilesToMakeLocalWritable, CancellationToken.None);
+
+			return 0;
+		}
+	}
+
+	class ConvertCLDataToLocalWritable : DataCommandBase
+	{
+		public override string Description => "Convert all data files in changelist to local writable";
+
+		public override CustomToolInfo CustomTool => new CustomToolInfo("Convert data to local writable", "%c");
+
+		public override async Task<int> Execute(string[] Args, IReadOnlyDictionary<string, string> ConfigValues, ILogger Logger)
+		{
+			int ChangeNumber;
+			if (Args.Length < 2)
+			{
+				Logger.LogError("Missing changelist number");
+				return 1;
+			}
+			else if (!int.TryParse(Args[1], out ChangeNumber))
+			{
+				Logger.LogError("'{Argument}' is not a numbered changelist", Args[1]);
+				return 1;
+			}
+
+			PerforceConnection Perforce = new PerforceConnection(null, null, Logger);
+			List<DescribeRecord> Describe = await Perforce.DescribeAsync(DescribeOptions.Shelved, -1, new int[] { ChangeNumber }, CancellationToken.None);
+			List<(string, string)> FilesInfo = Describe[0].Files.Select(d => (d.DepotFile, d.Type)).ToList();
+
+			return await MakeDataFilesLocalWritable(FilesInfo, Logger);
+		}
+	}
+
+	class ConvertDataToLocalWritable : DataCommandBase
+	{
+		public override string Description => "Convert all selected data files to local writable";
+
+		public override CustomToolInfo CustomTool => new CustomToolInfo("Convert data to local writable", "%F");
+
+		public override async Task<int> Execute(string[] Args, IReadOnlyDictionary<string, string> ConfigValues, ILogger Logger)
+		{
+			if (Args.Length < 2)
+			{
+				Logger.LogError("Missing file list");
+				return 1;
+			}
+
+			PerforceConnection Perforce = new PerforceConnection(null, null, Logger);
+			List<FStatRecord> StatRecords = await Perforce.FStatAsync(Args.Skip(1).ToList(), CancellationToken.None);
+			List<(string, string)> FilesInfo = StatRecords.Select(r => (r.DepotFile!, r.Type!)).ToList();
+
+			return await MakeDataFilesLocalWritable(FilesInfo, Logger);
+		}
+	}
+}
\ No newline at end of file
diff --git a/Engine/Extras/P4VUtils/Commands/UnshelveCommands.cs b/Engine/Extras/P4VUtils/Commands/UnshelveCommands.cs
new file mode 100644
--- /dev/null
+++ b/Engine/Extras/P4VUtils/Commands/UnshelveCommands.cs
@@ -0,0 +1,175 @@
+// Copyright Epic Games, Inc. All Rights Reserved.
+
+using EpicGames.Core;
+using EpicGames.Perforce;
+using Microsoft.Extensions.Logging;
+using System;
+using System.Collections.Generic;
+using System.ComponentModel;
+using System.IO;
+using System.Linq;
+using System.Text;
+using System.Threading;
+using System.Threading.Tasks;
+using System.Windows;
+
+namespace P4VUtils.Commands
+{
+	abstract class UnshelveCommandBase : Command
+	{
+
+		protected async Task<int> UnshelveFiles(List<(string, string, int)> FilesInfoToUnshelve, bool bMakeDataWritable, int ChangeNumber, int IntoChangeNumber, ILogger Logger)
+		{
+			PerforceConnection Perforce = new PerforceConnection(null, null, Logger);
+
+			List<string> FileSpecs = FilesInfoToUnshelve.Select(t => t.Item1).ToList();
+			List<FStatRecord> StatRecords = await Perforce.FStatAsync(FileSpecs, CancellationToken.None);
+			Dictionary<string, int> CurrentFileRevisions = StatRecords.ToDictionary(r => r.DepotFile!, r => r.HaveRevision);
+			Dictionary<string, string> CurrentDepotToClientFile = StatRecords.ToDictionary(r => r.DepotFile!, r => r.ClientFile!);
+			List<WhereRecord> WhereRecords = await Perforce.WhereAsync(FileSpecs, CancellationToken.None);
+
+			// Figure out the local path of the added file with a p4 where command
+			foreach (WhereRecord WhereRecord in WhereRecords)
+			{
+				if(!CurrentDepotToClientFile.TryGetValue(WhereRecord.DepotFile, out string? ClientFile) || ClientFile == null)
+				{
+					CurrentDepotToClientFile[WhereRecord.DepotFile] = WhereRecord.Path;
+				}
+			}
+
+			List<string> FilesToUnshelve;
+			if (bMakeDataWritable)
+			{
+				FilesToUnshelve = new List<string>();
+				FilesInfoToUnshelve.ForEach(async t =>
+				{
+					(string DepotFile, string Type, int Revision) = t;
+					if (Type.StartsWith("binary"))
+					{
+						if (CurrentDepotToClientFile.TryGetValue(DepotFile, out string? ClientFile) && ClientFile != null)
+						{
+							PrintRecord PrintRecord = await Perforce.PrintAsync(ClientFile, $"{DepotFile}@={ChangeNumber}", CancellationToken.None);
+							File.SetAttributes(ClientFile, ~FileAttributes.ReadOnly & File.GetAttributes(ClientFile));
+						}
+						else
+						{
+							Logger.LogError("Unshelve failed: could not map locally depot file:{DepotFile}");
+						}
+
+					}
+					else
+					{
+						FilesToUnshelve.Add(DepotFile);
+					}
+				});
+			}
+			else
+			{
+				FilesToUnshelve = FilesInfoToUnshelve.Select(t => t.Item1).ToList();
+			}
+
+			List<string> FilesToResolve = new List<string>();
+			List<string> FilesToSync = new List<string>();
+			FilesInfoToUnshelve.ForEach(f =>
+			{
+				(string DepotFile, string Type, int Revision) = f;
+				int CurrentFileRevision;
+				if (CurrentFileRevisions.TryGetValue(DepotFile, out CurrentFileRevision) && Revision < CurrentFileRevision)
+				{
+					FilesToSync.Add(DepotFile + "#" + CurrentFileRevision);
+					FilesToResolve.Add(DepotFile);
+				}
+			});
+
+			List<UnshelveRecord> UnshelveRecords = await Perforce.UnshelveAsync(ChangeNumber, IntoChangeNumber, null, null, null, UnshelveOptions.None, FilesToUnshelve, CancellationToken.None);
+			if (FilesToSync.Count > 0)
+			{
+				List<SyncRecord> SyncFiles = await Perforce.SyncAsync(FilesToSync, CancellationToken.None);
+				SyncFiles.ForEach(r =>
+				{
+					Logger.LogInformation("Restore file:{r.ClientFile} to revision:{r.Revision}");
+				});
+
+				List<ResolveRecord> ResolveRecords = await Perforce.ResolveAsync(-1, ResolveOptions.Automatic, FilesToResolve, CancellationToken.None);
+			}
+
+			return 0;
+		}
+
+		protected async Task<int> UnshelveChangeList(int ChangeNumber, bool bMakeDataWritable, ILogger Logger)
+		{
+			PerforceConnection Perforce = new PerforceConnection(null, null, Logger);
+			List<DescribeRecord> Describe = await Perforce.DescribeAsync(DescribeOptions.Shelved, -1, new int[] { ChangeNumber }, CancellationToken.None);
+
+			int IntoChangeNumber = -1;
+			InfoRecord Info = await Perforce.GetInfoAsync(InfoOptions.None, CancellationToken.None);
+			if (Describe[0].Client == Info.ClientName)
+			{
+				IntoChangeNumber = ChangeNumber;
+			}
+
+			if (Describe == null || Describe.Count != 1)
+			{
+				Logger.LogError("Unable to find changelist {Change}");
+				return 1;
+			}
+			if (Describe[0].Files.Count == 0)
+			{
+				Logger.LogError("No files are shelved in the given changelist{Change}");
+				return 1;
+			}
+
+			List<(string, string, int)> FilesInfoToUnshelve = Describe[0].Files.Select(d => (d.DepotFile, d.Type, d.Revision)).ToList();
+
+			return await UnshelveFiles(FilesInfoToUnshelve, bMakeDataWritable, ChangeNumber, IntoChangeNumber, Logger);
+		}
+	}
+
+	class UnshelveToCurrentRevision : UnshelveCommandBase
+	{
+		public override string Description => "Remembers revision all files you are about to unshelve and if the revision is older, will sync to saved revision";
+
+		public override CustomToolInfo CustomTool => new CustomToolInfo("Unshelve to current revision", "%p");
+
+		public override async Task<int> Execute(string[] Args, IReadOnlyDictionary<string, string> ConfigValues, ILogger Logger)
+		{
+			int ChangeNumber;
+			if (Args.Length < 2)
+			{
+				Logger.LogError("Missing changelist number");
+				return 1;
+			}
+			else if (!int.TryParse(Args[1], out ChangeNumber))
+			{
+				Logger.LogError("'{Argument}' is not a numbered changelist", Args[1]);
+				return 1;
+			}
+
+			return await UnshelveChangeList(ChangeNumber, false/*bMakeDataWritable*/, Logger);
+		}
+	}
+
+	class UnshelveMakeDataWritable : UnshelveCommandBase
+	{
+		public override string Description => "Unshelve changelist but for data will make a writable copy locally";
+
+		public override CustomToolInfo CustomTool => new CustomToolInfo("Unshelve and make data writable", "%p");
+
+		public override async Task<int> Execute(string[] Args, IReadOnlyDictionary<string, string> ConfigValues, ILogger Logger)
+		{
+			int ChangeNumber;
+			if (Args.Length < 2)
+			{
+				Logger.LogError("Missing changelist number");
+				return 1;
+			}
+			else if (!int.TryParse(Args[1], out ChangeNumber))
+			{
+				Logger.LogError("'{Argument}' is not a numbered changelist", Args[1]);
+				return 1;
+			}
+
+			return await UnshelveChangeList(ChangeNumber, true/*bMakeDataWritable*/, Logger);
+		}
+	}
+}
\ No newline at end of file
diff --git a/Engine/Extras/P4VUtils/Program.cs b/Engine/Extras/P4VUtils/Program.cs
--- a/Engine/Extras/P4VUtils/Program.cs
+++ b/Engine/Extras/P4VUtils/Program.cs
@@ -44,16 +44,25 @@ namespace P4VUtils
 
 	class Program
 	{
+		// UEHelpersInRoot - commands that help with common but simple operations
+		public static IReadOnlyDictionary<string, Command> RootHelperCommands { get; } = new Dictionary<string, Command>(StringComparer.OrdinalIgnoreCase)
+		{
+			["describe"] = new DescribeCommand(),
+			["copyclnum"] = new CopyCLCommand(),
+		};
+
 		// UEHelpers - commands that help with common but simple operations
 		public static IReadOnlyDictionary<string, Command> HelperCommands { get; } = new Dictionary<string, Command>(StringComparer.OrdinalIgnoreCase)
 		{
-			["describe"] = new DescribeCommand(),
 			["findlastedit"] = new FindLastEditCommand(),
 			["findlasteditbyline"] = new P4BlameCommand(),
 			["snapshot"] = new SnapshotCommand(),
-			["copyclnum"] = new CopyCLCommand(),
 			["reconcilecode"] = new FastReconcileCodeEditsCommand(),
 			["reconcileall"] = new FastReconcileAllEditsCommand(),
+			["unshelvetocurrentrevision"] = new UnshelveToCurrentRevision(),
+			["unshelvemakedatawritable"] = new UnshelveMakeDataWritable(),
+			["convertcldatatolocalwritable"] = new ConvertCLDataToLocalWritable(),
+			["convertdatatolocalwritable"] = new ConvertDataToLocalWritable(),
 		};
 
 		// UEIntegrate Folder commands - complex commands to facilitate integrations/backout
@@ -74,7 +83,7 @@ namespace P4VUtils
 			["movewriteablepreflightandsubmit"] = new MoveWriteableFilesthenPreflightAndSubmitCommand(),
 		};
 
-		public static IDictionary<string, Command> Commands = HelperCommands.Concat(IntegrateCommands).Concat(HordeCommands).ToDictionary(p => p.Key, p => p.Value, StringComparer.OrdinalIgnoreCase);
+		public static IDictionary<string, Command> Commands = RootHelperCommands.Concat(HelperCommands).Concat(IntegrateCommands).Concat(HordeCommands).ToDictionary(p => p.Key, p => p.Value, StringComparer.OrdinalIgnoreCase);
 
 		static void PrintHelp(ILogger Logger)
 		{
@@ -238,7 +247,7 @@ namespace P4VUtils
 			return ToolsChecked == ToolsRemoved;
 		}
 
-		static void InstallCommandsListInFolder(string FolderName, IReadOnlyDictionary<string, Command> InputCommmands, XmlDocument Document, FileReference DotNetLocation, FileReference AssemblyLocation, ILogger Logger)
+		static void InstallCommandsListInFolder(string FolderName, bool AddFolderToContextMenu, IReadOnlyDictionary<string, Command> InputCommmands, XmlDocument Document, FileReference DotNetLocation, FileReference AssemblyLocation, ILogger Logger)
 		{
 			// <CustomToolDefList>				// list of custom tools (top level)
 			//  < CustomToolDef >				// loose custom tool in top level
@@ -258,6 +267,10 @@ namespace P4VUtils
 				FolderDescription.InnerText = FolderName;
 				FolderDefinition.AppendChild(FolderDescription);
 
+				XmlElement FolderToContextMenu = Document.CreateElement("AddToContext");
+				FolderToContextMenu.InnerText = AddFolderToContextMenu ? "true" : "false";
+				FolderDefinition.AppendChild(FolderToContextMenu);
+
 				XmlElement FolderDefList = Document.CreateElement("CustomToolDefList");
 
 				foreach (KeyValuePair<string, Command> Pair in InputCommmands)
@@ -403,9 +416,10 @@ namespace P4VUtils
 			// Insert new entries
 			if (bInstall)
 			{
-				InstallCommandsListInFolder("UEHelpers", HelperCommands, Document, DotNetLocation, AssemblyLocation, Logger);
-				InstallCommandsListInFolder("UEIntegrate", IntegrateCommands, Document, DotNetLocation, AssemblyLocation, Logger);
-				InstallCommandsListInFolder("UEHorde", HordeCommands, Document, DotNetLocation, AssemblyLocation, Logger);
+				InstallCommandsListInFolder("UERootHelpers", false/*AddFolderToContextMenu*/, RootHelperCommands, Document, DotNetLocation, AssemblyLocation, Logger);
+				InstallCommandsListInFolder("UEHelpers", true/*AddFolderToContextMenu*/, HelperCommands, Document, DotNetLocation, AssemblyLocation, Logger);
+				InstallCommandsListInFolder("UEIntegrate", true/*AddFolderToContextMenu*/, IntegrateCommands, Document, DotNetLocation, AssemblyLocation, Logger);
+				InstallCommandsListInFolder("UEHorde", true/*AddFolderToContextMenu*/, HordeCommands, Document, DotNetLocation, AssemblyLocation, Logger);
 			}
 
 			// Save the new document
diff --git a/Engine/Source/Programs/Shared/EpicGames.Perforce/PerforceConnection.cs b/Engine/Source/Programs/Shared/EpicGames.Perforce/PerforceConnection.cs
--- a/Engine/Source/Programs/Shared/EpicGames.Perforce/PerforceConnection.cs
+++ b/Engine/Source/Programs/Shared/EpicGames.Perforce/PerforceConnection.cs
@@ -3152,7 +3152,7 @@ namespace EpicGames.Perforce
 		/// <param name="FileSpecs">Files to unshelve</param>
 		/// <param name="CancellationToken">Token used to cancel the operation</param>
 		/// <returns>Response from the server</returns>
-		public static async Task<List<UnshelveRecord>> UnshelveAsync(this IPerforceConnection Connection, int ChangeNumber, int IntoChangeNumber, string UsingBranchSpec, string UsingStream, string ForceParentStream, UnshelveOptions Options, FileSpecList FileSpecs, CancellationToken CancellationToken = default)
+		public static async Task<List<UnshelveRecord>> UnshelveAsync(this IPerforceConnection Connection, int ChangeNumber, int IntoChangeNumber, string? UsingBranchSpec, string? UsingStream, string? ForceParentStream, UnshelveOptions Options, FileSpecList FileSpecs, CancellationToken CancellationToken = default)
 		{
 			return (await TryUnshelveAsync(Connection, ChangeNumber, IntoChangeNumber, UsingBranchSpec, UsingStream, ForceParentStream, Options, FileSpecs, CancellationToken)).Data;
 		}
@@ -3170,7 +3170,7 @@ namespace EpicGames.Perforce
 		/// <param name="FileSpecs">Files to unshelve</param>
 		/// <param name="CancellationToken">Token used to cancel the operation</param>
 		/// <returns>Response from the server</returns>
-		public static Task<PerforceResponseList<UnshelveRecord>> TryUnshelveAsync(this IPerforceConnection Connection, int ChangeNumber, int IntoChangeNumber, string UsingBranchSpec, string UsingStream, string ForceParentStream, UnshelveOptions Options, FileSpecList FileSpecs, CancellationToken CancellationToken = default)
+		public static Task<PerforceResponseList<UnshelveRecord>> TryUnshelveAsync(this IPerforceConnection Connection, int ChangeNumber, int IntoChangeNumber, string? UsingBranchSpec, string? UsingStream, string? ForceParentStream, UnshelveOptions Options, FileSpecList FileSpecs, CancellationToken CancellationToken = default)
 		{
 			List<string> Arguments = new List<string>();
 			Arguments.Add($"-s{ChangeNumber}");
