P4VUtils: Fix compile errors after recent changes to EpicGames.Perforce library.

diff --git a/Engine/Extras/P4VUtils/Commands/CherryPickCommand.cs b/Engine/Extras/P4VUtils/Commands/CherryPickCommand.cs
--- a/Engine/Extras/P4VUtils/Commands/CherryPickCommand.cs
+++ b/Engine/Extras/P4VUtils/Commands/CherryPickCommand.cs
@@ -74,7 +74,7 @@ namespace P4VUtils.Commands
 			await Perforce.MergeAsync(MergeOptions.None, NewChangeRecord.Number, -1, SourceFileSpec, TargetFileSpec, CancellationToken.None);
 			Logger.LogInformation("Merged files into changelist {Change}", NewChangeRecord.Number);
 
-			PerforceResponseList<ResolveRecord> ResolveRecords = await Perforce.TryResolveAsync(NewChangeRecord.Number, ResolveOptions.Automatic, FileSpecList.Empty, CancellationToken.None);
+			PerforceResponseList<ResolveRecord> ResolveRecords = await Perforce.TryResolveAsync(NewChangeRecord.Number, ResolveOptions.Automatic, FileSpecList.Any, CancellationToken.None);
 			if (!ResolveRecords.Succeeded)
 			{
 				Logger.LogError("Unable to resolve files. Please resolve manually.");
diff --git a/Engine/Extras/P4VUtils/Commands/CompileCommand.cs b/Engine/Extras/P4VUtils/Commands/CompileCommand.cs
--- a/Engine/Extras/P4VUtils/Commands/CompileCommand.cs
+++ b/Engine/Extras/P4VUtils/Commands/CompileCommand.cs
@@ -50,7 +50,7 @@ namespace P4VUtils.Commands
 
 			List<FileReference> LocalFiles = new List<FileReference>();
 
-			List<WhereRecord> WhereRecords = await Perforce.WhereAsync(ChangeRecord.Files.ToArray(), CancellationToken.None);
+			List<WhereRecord> WhereRecords = await Perforce.WhereAsync(ChangeRecord.Files.ToArray(), CancellationToken.None).ToListAsync();
 			foreach (WhereRecord WhereRecord in WhereRecords)
 			{
 				if(WhereRecord.Path != null)
diff --git a/Engine/Extras/P4VUtils/Commands/ConvertToEditCommand.cs b/Engine/Extras/P4VUtils/Commands/ConvertToEditCommand.cs
--- a/Engine/Extras/P4VUtils/Commands/ConvertToEditCommand.cs
+++ b/Engine/Extras/P4VUtils/Commands/ConvertToEditCommand.cs
@@ -47,7 +47,7 @@ namespace P4VUtils.Commands
 		{
 			// Get the list of opened files in the CL along with their metadata
 			Logger.LogInformation("Getting the list of files in changelist {Change}...", Change);
-			List<FStatRecord> OpenedRecords = await Perforce.GetOpenFilesAsync(OpenedOptions.None, Change, null, null, -1, FileSpecList.Empty, CancellationToken.None);
+			List<OpenedRecord> OpenedRecords = await Perforce.OpenedAsync(OpenedOptions.None, Change, null, null, -1, FileSpecList.Any, CancellationToken.None).ToListAsync();
 			if (OpenedRecords.Count == 0)
 			{
 				Logger.LogError("Change has no opened files. Aborting.");
@@ -58,13 +58,13 @@ namespace P4VUtils.Commands
 			Logger.LogInformation("There are {NumFiles} open files in the change.", OpenedRecords.Count);
 
 			// create a mapping of client to local files so we can re-open using local paths (needed for some actions)
-			List<WhereRecord> WhereRecords = await Perforce.WhereAsync(OpenedRecords.Select(x => x.DepotFile!).ToArray(), CancellationToken.None);
+			List<WhereRecord> WhereRecords = await Perforce.WhereAsync(OpenedRecords.Select(x => x.DepotFile!).ToArray(), CancellationToken.None).ToListAsync();
 			Dictionary<string, string> OpenedClientToLocalMap = WhereRecords.ToDictionary(x => x.ClientFile, x => x.Path, StringComparer.OrdinalIgnoreCase);
 
 			// If any file that needs to be resolved still, we need to abort.
 			// We don't want to convert non-resolved files to plain edits,
 			// as we need to resolve the merge first!
-			List<ResolveRecord> ResolveRecords = await Perforce.ResolveAsync(Change, ResolveOptions.PreviewOnly, FileSpecList.Empty!, CancellationToken.None);
+			List<ResolveRecord> ResolveRecords = await Perforce.ResolveAsync(Change, ResolveOptions.PreviewOnly, FileSpecList.Any, CancellationToken.None);
 			if (ResolveRecords.Any())
 			{
 				Logger.LogError("The following files in the changelist need to be resolved.");
@@ -80,15 +80,15 @@ namespace P4VUtils.Commands
 			// with the same action/filetype. This should remove the integration records.
 
 			// partition our files into ones we know the actions for, and the ones we don't
-			List<FStatRecord> DeleteActions = new List<FStatRecord>();
-			List<FStatRecord> AddActions = new List<FStatRecord>();
-			List<FStatRecord> EditActions = new List<FStatRecord>();
-			List<FStatRecord> MoveAddActions = new List<FStatRecord>();
-			List<FStatRecord> MoveDeleteActions = new List<FStatRecord>();
-			List<FStatRecord> UnknownActions = new List<FStatRecord>();
+			List<OpenedRecord> DeleteActions = new List<OpenedRecord>();
+			List<OpenedRecord> AddActions = new List<OpenedRecord>();
+			List<OpenedRecord> EditActions = new List<OpenedRecord>();
+			List<OpenedRecord> MoveAddActions = new List<OpenedRecord>();
+			List<OpenedRecord> MoveDeleteActions = new List<OpenedRecord>();
+			List<OpenedRecord> UnknownActions = new List<OpenedRecord>();
 
 			// Multiple action map to the same reopen mapping, so this handles that.
-			Dictionary<FileAction, List<FStatRecord>> ActionMappings = new Dictionary<FileAction, List<FStatRecord>>()
+			Dictionary<FileAction, List<OpenedRecord>> ActionMappings = new Dictionary<FileAction, List<OpenedRecord>>()
 			{
 				[FileAction.Delete] = DeleteActions,
 				[FileAction.Branch] = AddActions,
@@ -100,9 +100,9 @@ namespace P4VUtils.Commands
 			};
 
 			// this maps all open records to a reopen action.
-			foreach (FStatRecord OpenedRecord in OpenedRecords)
+			foreach (OpenedRecord OpenedRecord in OpenedRecords)
 			{
-				List<FStatRecord>? ActionList;
+				List<OpenedRecord>? ActionList;
 				if (!ActionMappings.TryGetValue(OpenedRecord.Action, out ActionList))
 				{
 					ActionList = UnknownActions;
@@ -114,7 +114,7 @@ namespace P4VUtils.Commands
 			if (UnknownActions.Count > 0)
 			{
 				Logger.LogError("These files have unknown actions so cannot be reliably re-opened:");
-				foreach (FStatRecord OpenedRecord in UnknownActions)
+				foreach (OpenedRecord OpenedRecord in UnknownActions)
 				{
 					Logger.LogError("  {Action}: {ClientFile}", OpenedRecord.Action.ToString(), OpenedRecord.ClientFile);
 				}
@@ -133,17 +133,17 @@ namespace P4VUtils.Commands
 			// if the haveRev doesn't match the rev of the originally opened file, this was the case.
 			// so we sync -k because the local file matches the one we want, we just have to tell P4 that so it will let us check it out at that revision.
 			// If we simply sync -k to #head we might miss a legitimate submit that happened since our last resolve, and that would stomp that submit.
-			List<FStatRecord> NotSyncedOpenFiles = OpenedRecords.Where(x => x.HaveRevision != x.Revision && x.Action != FileAction.Branch && x.Action != FileAction.Add).ToList();
+			List<OpenedRecord> NotSyncedOpenFiles = OpenedRecords.Where(x => x.HaveRevision != x.Revision && x.Action != FileAction.Branch && x.Action != FileAction.Add).ToList();
 			if (NotSyncedOpenFiles.Count > 0)
 			{
 				Logger.LogWarning("The following files were not synced to the revision they were resolved to before. syncing them to the indicated rev first.");
-				foreach (FStatRecord NotSyncedOpenFile in NotSyncedOpenFiles)
+				foreach (OpenedRecord NotSyncedOpenFile in NotSyncedOpenFiles)
 				{
 					Logger.LogWarning("    {File}#{Revision} (had #{HaveRevision})", NotSyncedOpenFile.ClientFile, NotSyncedOpenFile.Revision, NotSyncedOpenFile.HaveRevision);
 				}
 				if (!Debug)
 				{
-					await Perforce.SyncAsync(SyncOptions.KeepWorkspaceFiles, -1, NotSyncedOpenFiles.Select(x => $"{x.ClientFile}#{x.Revision}").ToArray(), CancellationToken.None);
+					await Perforce.SyncAsync(SyncOptions.KeepWorkspaceFiles, -1, NotSyncedOpenFiles.Select(x => $"{x.ClientFile}#{x.Revision}").ToArray(), CancellationToken.None).ToListAsync();
 				}
 			}
 
@@ -169,7 +169,7 @@ namespace P4VUtils.Commands
 					await Perforce.EditAsync(Change, null, EditOptions.KeepWorkspaceFiles, MoveAddActions.Select(x => x.MovedFile!).ToArray(), CancellationToken.None);
 
 					// then we can open the file for move in the new location (have to use -k because the file has already been moved locally!)
-					foreach ((FStatRecord OpenedRecord, string LocalFile) in MoveAddActions.Zip(LocalFiles))
+					foreach ((OpenedRecord OpenedRecord, string LocalFile) in MoveAddActions.Zip(LocalFiles))
 					{
 						await Perforce.MoveAsync(Change, null, MoveOptions.KeepWorkspaceFiles, OpenedRecord.MovedFile!, LocalFile, CancellationToken.None);
 					}
@@ -178,7 +178,7 @@ namespace P4VUtils.Commands
 
 			// Get the list of reopened files in the CL to check their filetype
 			Logger.LogInformation("Checking the list of files reopened in changelist {Change}...", Change);
-			List<FStatRecord> ReopenedRecords = await Perforce.GetOpenFilesAsync(OpenedOptions.None, Change, null, null, -1, FileSpecList.Empty!, CancellationToken.None);
+			List<OpenedRecord> ReopenedRecords = await Perforce.OpenedAsync(OpenedOptions.None, Change, null, null, -1, FileSpecList.Any, CancellationToken.None).ToListAsync();
 			if (ReopenedRecords.Count == 0)
 			{
 				Logger.LogError("Change has no reopened files. This signifies an error in the script! Aborting...");
@@ -191,10 +191,10 @@ namespace P4VUtils.Commands
 			}
 
 			// Ensure the filetypes match and ensure each reopened file was in the original changelist.
-			Dictionary<string, FStatRecord> ReopenedRecordsMap = ReopenedRecords.ToDictionary(x => x.ClientFile!, x => x);
-			foreach (FStatRecord OpenedRecord in OpenedRecords)
+			Dictionary<string, OpenedRecord> ReopenedRecordsMap = ReopenedRecords.ToDictionary(x => x.ClientFile!, x => x);
+			foreach (OpenedRecord OpenedRecord in OpenedRecords)
 			{
-				FStatRecord? ReopenedRecord;
+				OpenedRecord? ReopenedRecord;
 				if (!ReopenedRecordsMap.TryGetValue(OpenedRecord.ClientFile!, out ReopenedRecord))
 				{
 					Logger.LogError("Could not find original file {ClientFile} in re-opened records. This signifies an error in the script! Aborting...", OpenedRecord.ClientFile);
@@ -215,7 +215,7 @@ namespace P4VUtils.Commands
 			return true;
 		}
 
-		static async Task ReopenFiles(string Operation, List<FStatRecord> Records, Dictionary<string, string> OpenedClientToLocalMap, Func<string[], Task> ReopenAsync, bool Debug, ILogger Logger)
+		static async Task ReopenFiles(string Operation, List<OpenedRecord> Records, Dictionary<string, string> OpenedClientToLocalMap, Func<string[], Task> ReopenAsync, bool Debug, ILogger Logger)
 		{
 			if (Records.Count > 0)
 			{
diff --git a/Engine/Extras/P4VUtils/Commands/DataCommands.cs b/Engine/Extras/P4VUtils/Commands/DataCommands.cs
--- a/Engine/Extras/P4VUtils/Commands/DataCommands.cs
+++ b/Engine/Extras/P4VUtils/Commands/DataCommands.cs
@@ -72,7 +72,7 @@ namespace P4VUtils.Commands
 			}
 
 			using PerforceConnection Perforce = new PerforceConnection(null, null, Logger);
-			List<FStatRecord> StatRecords = await Perforce.FStatAsync(Args.Skip(1).ToList(), CancellationToken.None);
+			List<FStatRecord> StatRecords = await Perforce.FStatAsync(Args.Skip(1).ToList(), CancellationToken.None).ToListAsync();
 			List<(string, string)> FilesInfo = StatRecords.Select(r => (r.DepotFile!, r.Type!)).ToList();
 
 			return await MakeDataFilesLocalWritable(FilesInfo, Logger);
diff --git a/Engine/Extras/P4VUtils/Commands/PreflightCommand.cs b/Engine/Extras/P4VUtils/Commands/PreflightCommand.cs
--- a/Engine/Extras/P4VUtils/Commands/PreflightCommand.cs
+++ b/Engine/Extras/P4VUtils/Commands/PreflightCommand.cs
@@ -54,7 +54,7 @@ namespace P4VUtils.Commands
 				return 1;
 			}
 
-			List<FStatRecord> OpenedRecords = await Perforce.GetOpenFilesAsync(OpenedOptions.None, Change, null, null, -1, FileSpecList.Empty, CancellationToken.None);
+			List<OpenedRecord> OpenedRecords = await Perforce.OpenedAsync(OpenedOptions.None, Change, null, null, -1, FileSpecList.Any, CancellationToken.None).ToListAsync();
 
 			if (OpenedRecords.Count > 0)
 			{
@@ -93,7 +93,7 @@ namespace P4VUtils.Commands
 
 					Logger.LogInformation("Created pending changelist {Change}", NewChangeRecord.Number);
 
-					foreach (FStatRecord OpenedRecord in OpenedRecords)
+					foreach (OpenedRecord OpenedRecord in OpenedRecords)
 					{
 						if (OpenedRecord.ClientFile != null)
 						{
diff --git a/Engine/Extras/P4VUtils/Commands/SnapShotCommand.cs b/Engine/Extras/P4VUtils/Commands/SnapShotCommand.cs
--- a/Engine/Extras/P4VUtils/Commands/SnapShotCommand.cs
+++ b/Engine/Extras/P4VUtils/Commands/SnapShotCommand.cs
@@ -8,6 +8,7 @@ using System.Collections.Generic;
 using System.ComponentModel;
 using System.Diagnostics;
 using System.IO;
+using System.Linq;
 using System.Reflection;
 using System.Runtime.InteropServices;
 using System.Text;
@@ -48,7 +49,7 @@ namespace P4VUtils.Commands
 			}
 
 			// Create the copy of the CL IFF it has files opened
-			List<FStatRecord> OpenedRecords = await Perforce.GetOpenFilesAsync(OpenedOptions.None, Change, null, null, -1, FileSpecList.Empty, CancellationToken.None);
+			List<OpenedRecord> OpenedRecords = await Perforce.OpenedAsync(OpenedOptions.None, Change, null, null, -1, FileSpecList.Any, CancellationToken.None).ToListAsync();
 
 			if (OpenedRecords.Count > 0)
 			{
@@ -66,7 +67,7 @@ namespace P4VUtils.Commands
 				Logger.LogInformation("Created pending changelist {Change}", NewChangeRecord.Number);
 
 				// Move the files to the new CL
-				foreach (FStatRecord OpenedRecord in OpenedRecords)
+				foreach (OpenedRecord OpenedRecord in OpenedRecords)
 				{
 					if (OpenedRecord.ClientFile != null)
 					{
@@ -79,7 +80,7 @@ namespace P4VUtils.Commands
 				await Perforce.ShelveAsync(NewChangeRecord.Number, ShelveOptions.Overwrite, new[] { "//..." }, CancellationToken.None);
 
 				// Move the files BACK to the old CL
-				foreach (FStatRecord OpenedRecord in OpenedRecords)
+				foreach (OpenedRecord OpenedRecord in OpenedRecords)
 				{
 					if (OpenedRecord.ClientFile != null)
 					{
diff --git a/Engine/Extras/P4VUtils/Commands/UnshelveCommands.cs b/Engine/Extras/P4VUtils/Commands/UnshelveCommands.cs
--- a/Engine/Extras/P4VUtils/Commands/UnshelveCommands.cs
+++ b/Engine/Extras/P4VUtils/Commands/UnshelveCommands.cs
@@ -23,10 +23,10 @@ namespace P4VUtils.Commands
 			using PerforceConnection Perforce = new PerforceConnection(null, null, Logger);
 
 			List<string> FileSpecs = FilesInfoToUnshelve.Select(t => t.Item1).ToList();
-			List<FStatRecord> StatRecords = await Perforce.FStatAsync(FileSpecs, CancellationToken.None);
+			List<FStatRecord> StatRecords = await Perforce.FStatAsync(FileSpecs, CancellationToken.None).ToListAsync();
 			Dictionary<string, int> CurrentFileRevisions = StatRecords.ToDictionary(r => r.DepotFile!, r => r.HaveRevision);
 			Dictionary<string, string> CurrentDepotToClientFile = StatRecords.ToDictionary(r => r.DepotFile!, r => r.ClientFile!);
-			List<WhereRecord> WhereRecords = await Perforce.WhereAsync(FileSpecs, CancellationToken.None);
+			List<WhereRecord> WhereRecords = await Perforce.WhereAsync(FileSpecs, CancellationToken.None).ToListAsync();
 
 			// Figure out the local path of the added file with a p4 where command
 			foreach (WhereRecord WhereRecord in WhereRecords)
@@ -84,10 +84,10 @@ namespace P4VUtils.Commands
 			List<UnshelveRecord> UnshelveRecords = await Perforce.UnshelveAsync(ChangeNumber, IntoChangeNumber, null, null, null, UnshelveOptions.None, FilesToUnshelve, CancellationToken.None);
 			if (FilesToSync.Count > 0)
 			{
-				List<SyncRecord> SyncFiles = await Perforce.SyncAsync(FilesToSync, CancellationToken.None);
+				List<SyncRecord> SyncFiles = await Perforce.SyncAsync(FilesToSync, CancellationToken.None).ToListAsync();
 				SyncFiles.ForEach(r =>
 				{
-					Logger.LogInformation("Restore file:{ClientFile} to revision:{Revision}", r.ClientFile, r.Revision);
+					Logger.LogInformation("Restore file:{ClientFile} to revision:{Revision}", r.Path, r.Revision);
 				});
 
 				List<ResolveRecord> ResolveRecords = await Perforce.ResolveAsync(-1, ResolveOptions.Automatic, FilesToResolve, CancellationToken.None);
diff --git a/Engine/Source/Programs/Shared/EpicGames.Perforce/Records/OpenedRecord.cs b/Engine/Source/Programs/Shared/EpicGames.Perforce/Records/OpenedRecord.cs
--- a/Engine/Source/Programs/Shared/EpicGames.Perforce/Records/OpenedRecord.cs
+++ b/Engine/Source/Programs/Shared/EpicGames.Perforce/Records/OpenedRecord.cs
@@ -27,6 +27,12 @@ namespace EpicGames.Perforce
 		[PerforceTag("clientFile", Optional = true)]
 		public string ClientFile { get; set; } = String.Empty;
 
+		/// <summary>
+		/// For a move action (either move/add or move/delete) indicates the counterpart depot path
+		/// </summary>
+		[PerforceTag("movedFile", Optional = true)]
+		public string? MovedFile { get; set; }
+
 		/// <summary>
 		/// The revision of the file (not returned if OpenedOptions.ShortOutput is set)
 		/// </summary>
@@ -51,6 +57,12 @@ namespace EpicGames.Perforce
 		[PerforceTag("change")]
 		public int Change { get; set; }
 
+		/// <summary>
+		/// New filetype for the file
+		/// </summary>
+		[PerforceTag("type")]
+		public string Type { get; set; } = String.Empty;
+
 		/// <summary>
 		/// User with the file open
 		/// </summary>
