- Fix P4VUtils to compile with latest EpicGames.Core

diff --git a/Engine/Extras/P4VUtils/Commands/P4BlameCommand.cs b/Engine/Extras/P4VUtils/Commands/P4BlameCommand.cs
--- a/Engine/Extras/P4VUtils/Commands/P4BlameCommand.cs
+++ b/Engine/Extras/P4VUtils/Commands/P4BlameCommand.cs
@@ -54,7 +54,7 @@ namespace P4VUtils.Commands
 				AnnotateRecord LineResult = Results[LineNumber];
 				string EscapedString = Regex.Escape(LineResult.Data);
 
-				Logger.LogInformation(@"Change {Change} by {User} on [{Time}]:'{Description}'", LineResult.LowerCL, LineResult.UserName, LineResult.Time, EscapedString);
+				Logger.LogInformation(@"Change {Change} by {User} on [{Time}]:'{Description}'", LineResult.LowerCl, LineResult.UserName, LineResult.Time, EscapedString);
 			}
 			catch (PerforceException Exception)
 			{
diff --git a/Engine/Extras/P4VUtils/Program.cs b/Engine/Extras/P4VUtils/Program.cs
--- a/Engine/Extras/P4VUtils/Program.cs
+++ b/Engine/Extras/P4VUtils/Program.cs
@@ -101,7 +101,7 @@ namespace P4VUtils
 			}
 
 			Logger.LogInformation("Commands:");
-			HelpUtils.PrintTable(Table, 2, 15, Logger);
+			HelpUtils.PrintTable(Table, 2, 15, ConsoleUtils.WindowWidth - 1, Logger);
 		}
 
 		static async Task<int> Main(string[] Args)
@@ -149,7 +149,10 @@ namespace P4VUtils
 				if (Args.Any(x => x.Equals("-help", StringComparison.OrdinalIgnoreCase)))
 				{
 					List<KeyValuePair<string, string>> Parameters = CommandLineArguments.GetParameters(Command.GetType());
-					HelpUtils.PrintHelp(Args[0], Command.GetType(), Logger);
+					Logger.LogInformation("{Title}", Args[0]);
+					Logger.LogInformation("{Description}", Command.GetType());
+					Logger.LogInformation("Parameters:");
+					HelpUtils.PrintTable(Parameters, 4, 24, HelpUtils.WindowWidth - 1, Logger);
 					return 0;
 				}
 
diff --git a/Engine/Source/Programs/Shared/EpicGames.Core/HelpUtils.cs b/Engine/Source/Programs/Shared/EpicGames.Core/HelpUtils.cs
--- a/Engine/Source/Programs/Shared/EpicGames.Core/HelpUtils.cs
+++ b/Engine/Source/Programs/Shared/EpicGames.Core/HelpUtils.cs
@@ -141,7 +141,7 @@ namespace EpicGames.Core
 
 			foreach (string line in lines)
 			{
-				logger.LogInformation("{Line)", line);
+				logger.LogInformation("{Line}", line);
 			}
 		}
 
