Extend P4VUtils with a new command that allows the user to submit&virtualize package files (for projects with asset virtualization enabled)

#rnx

### P4VUtils
- Add a new submenu "UESubmit" for commands related to the submission of files and changelists
-At the moment the submenu only has a single command SubmitAndVirtualizeCommand

### SubmitAndVirtualizeCommand
- This command is used to run the asset virtualization process over all of the files in a changelist before submitting it using the following logic:
-- Find all package files in a changelist
-- Sort the files by project
-- Find the engine installation associated with each project
-- Build and run the stand alone virtualization tool (if needed) and pass it the packages for the project
-- If the virtualization process succeeeded for all projects we can submit the original changelist, other wise print errors to the p4v terminal.
- In most cases all of the packages in a changelist will be under a single project, but the command should be able to cope with many packages from many different projects using many different engine installations.
- Note that unlike a normal submit, virtualizing a package can change it on disk (by removing the payload data) which means we will probably run into problems where a running instance of the editor has access locks on the file and they cannot be edited. At the moment this command and the stand alone tool will just error, there is no tool<->editor communication protocol set up to request that the editor release the package file. This will be addressed as a future work item.
- This is the first pass of the command, there are a number of non-essential todo comments left in the code that will be addressed based on feedback from the end users.

diff --git a/Engine/Extras/P4VUtils/Commands/SubmitCommand.cs b/Engine/Extras/P4VUtils/Commands/SubmitCommand.cs
new file mode 100644
--- /dev/null
+++ b/Engine/Extras/P4VUtils/Commands/SubmitCommand.cs
@@ -0,0 +1,685 @@
+// Copyright Epic Games, Inc. All Rights Reserved.
+
+using EpicGames.Core;
+using EpicGames.Perforce;
+using Microsoft.Extensions.Logging;
+using Microsoft.Win32;
+using System;
+using System.Collections.Generic;
+using System.IO;
+using System.Linq;
+using System.Text;
+using System.Threading;
+using System.Threading.Tasks;
+
+namespace P4VUtils.Commands
+{
+	class SubmitAndVirtualizeCommand : Command
+	{
+		public override string Description => "Virtualize And Submit";
+
+		public override CustomToolInfo CustomTool => new CustomToolInfo("Virtualize And Submit", "$c %c") { ShowConsole = true, RefreshUI = true };
+
+		public override async Task<int> Execute(string[] args, IReadOnlyDictionary<string, string> configValues, ILogger logger)
+		{
+			// Parse command lines
+			if (args.Length < 3)
+			{
+				logger.LogError("Not enough args for command!");
+				return 1;
+			}
+
+			string clientSpec = args[1];
+
+			if (!int.TryParse(args[2], out int changeNumber))
+			{
+				logger.LogError("'{Argument}' is not a numbered changelist", args[1]);
+				return 1;
+			}
+
+			logger.LogInformation("Attempting to virtualize and submit changelist {Change} in the workspace {Spec}", changeNumber, clientSpec);
+
+			// Connect to perforce and validate
+
+			logger.LogInformation("Connecting to perforce...");
+			
+			// We prefer the native client to avoid the problem where different versions of p4.exe expect
+			// or return records with different formatting to each other.
+			PerforceSettings settings = new PerforceSettings(PerforceSettings.Default) { PreferNativeClient = true, ClientName = clientSpec };
+			using IPerforceConnection perforceConnection = await PerforceConnection.CreateAsync(settings, logger);
+			if (perforceConnection == null)
+			{
+				logger.LogError("Failed to connect", args[1]);
+				return 1;
+			}
+
+			// First we need to find the packages in the changelist
+			// First we need to find the packages in the changelist
+			string[]? localFilePaths = await FindPackagesInChangelist(perforceConnection, changeNumber, logger);
+			if(localFilePaths == null)
+			{
+				return 1;
+			}
+			logger.LogInformation("Found {Amount} package(s) that may need virtualization", localFilePaths.Length);
+
+			// Now sort the package paths by their project (it is unlikely but a user could be submitting content
+			// from multiple projects that the same time)
+
+			Dictionary<string, List<string>> projects = SortPackagesIntoProjects(localFilePaths, logger);
+
+			logger.LogInformation("The packages are distributed between {Amount} project(s)", projects.Count);
+
+			// Find engine per project
+			IReadOnlyDictionary<string, string> engineInstalls = EnumerateEngineInstallations(logger);
+
+			foreach(KeyValuePair<string, List<string>> project in projects)
+			{
+				// Get engine association from project file
+				string engineIdentifier = GetEngineIdentifierForProject(project.Key);
+				if (!String.IsNullOrEmpty(engineIdentifier))
+				{
+					string engineRoot = GetEngineRootDirFromIdentifier(engineIdentifier);
+					if (!String.IsNullOrEmpty(engineRoot))
+					{
+						logger.LogInformation("\nAttempting to virtualize packages for project '{Project}' with the engine installation '{Engine}'", project.Key, engineRoot);
+						// @todo Many projects can share the same engine install, and technically UnrealVirtualizationTool
+						// supports the virtualization files from many projects at the same time. We could consider doing
+						// this pass per engine install rather than per project? At the very least we should only 'build'
+						// the tool once per engine
+
+						Task<bool> compileResult = BuildVirtualizationTool(engineRoot, logger);
+
+						string tempFilesPath = await WritePackageFileList(project.Value, logger);
+
+						// Check if the compilation of the tool succeeded or not
+						if (await compileResult == false)
+						{
+							return 1;
+						}
+
+						if(await RunVirtualizationTool(engineRoot, clientSpec, tempFilesPath, logger) == false)
+						{
+							return 1;
+						};
+
+						// @todo Make sure this always gets cleaned up when we go out of scope
+						File.Delete(tempFilesPath);
+					}
+				}
+			}
+
+			ChangeRecord? changeRecord = await StampChangelistDescription(perforceConnection, changeNumber, logger);
+			if (changeRecord == null)
+			{
+				return 1;
+			}
+
+			logger.LogInformation("Virtualization completed, attempting to submit changelist {Number}", changeNumber);
+
+			// Submit
+			//@todo ideally we should get the tags back from UnrealVirtualizationTool
+			PerforceResponseList<SubmitRecord> submitResponses = await TrySubmitAsync(perforceConnection, changeNumber, SubmitOptions.None, CancellationToken.None);
+
+			bool successfulSubmit = submitResponses.All(x => x.Succeeded);
+			if (successfulSubmit)
+			{
+				// @todo The submit API really should return the number that the changelist was finally submit as so that
+				// we can log that here instead.
+				logger.LogInformation("Successfully submited changelist {Change}", changeNumber);
+			}
+			else
+			{
+				// Log every response that was a failure and has an error message associated with it
+				logger.LogError("Submit failed due to:");
+				foreach (PerforceResponse<SubmitRecord> response in submitResponses)
+				{
+					if (response.Failed && response.Error != null)
+					{
+						logger.LogError("{Message}", response.Error.ToString());
+					}
+				}
+
+				// If the final submit failed we remove the virtualization tag, even though the changelist is technically
+				// virtualized at this point and submitting it would be safe.
+				// This is to keep the behavior the same as the native code paths.
+				logger.LogInformation("Removing virtualization tags from the changelist..");
+				PerforceResponse updateResponse = await perforceConnection.TryUpdateChangeAsync(UpdateChangeOptions.None, changeRecord, CancellationToken.None);
+				if (!updateResponse.Succeeded)
+				{
+					logger.LogError("Failed to remove the virtualization tags!");
+				}		
+
+				return 1;
+			}
+		
+			return 0;
+		}
+
+		/// <summary>
+		/// Compiles UnrealVirtualizationTool via RunUBT.bat
+		/// </summary>
+		/// <param name="engineRoot">Root path of the engine we want to build the tool for</param>
+		/// <param name="logger">Interface for logging</param>
+		/// <returns>True if the tool built successfully, otherwise false</returns>
+		private static async Task<bool> BuildVirtualizationTool(string engineRoot, ILogger logger)
+		{
+			logger.LogInformation("Building UnrealVirtualizationTool...");
+
+			string buildBatchFile = Path.Combine(engineRoot, @"Engine\Build\BatchFiles\RunUBT.bat");
+			StringBuilder arguments = new StringBuilder($"{buildBatchFile.QuoteArgument()}");
+			arguments.Append(" UnrealVirtualizationTool Win64 development -Progress");
+
+			string shellFileName = Environment.GetEnvironmentVariable("COMSPEC") ?? "C:\\Windows\\System32\\cmd.exe";
+			string shellArguments = $"/d/s/c \"{arguments}\"";
+
+			using (MemoryStream bufferedOutput = new MemoryStream())
+			using (ManagedProcessGroup Group = new ManagedProcessGroup())
+			using (ManagedProcess Process = new ManagedProcess(Group, shellFileName, shellArguments, null, null, System.Diagnostics.ProcessPriorityClass.Normal))
+			{
+				await Process.CopyToAsync(bufferedOutput, CancellationToken.None);
+
+				// We only show the output if there was an error to avoid showing too much
+				// info to the user when they are invoking this from p4v.
+				// We need to print everything from stdout and stderr as stderr alone often 
+				// does not contain all of the build failure info. (note that ManagedProcess is
+				// merging the two streams by default)
+				if (Process.ExitCode != 0)
+				{
+					bufferedOutput.Seek(0, SeekOrigin.Begin);
+
+					using (Stream stdOutput = Console.OpenStandardOutput())
+					{
+						await bufferedOutput.CopyToAsync(stdOutput, CancellationToken.None);
+					}
+
+					logger.LogError("Failed to build UnrealVirtualizationTool");
+					return false;
+				}
+			}
+
+			logger.LogInformation("UnrealVirtualizationTool built successfully");
+			return true;
+		}
+
+		/// <summary>
+		/// Runs the UnrealVirtualizationTool
+		/// </summary>
+		/// <param name="engineRoot">Root path of the engine we want to run the tool from</param>
+		/// <param name="clientSpec">The perforce client spec that the files are in</param>
+		/// <param name="packageListPath">A path to a text file containing the paths of the packages to be virtualized</param>
+		/// <param name="logger">Interface for logging</param>
+		/// <returns></returns>
+		private static async Task<bool> RunVirtualizationTool(string engineRoot, string clientSpec, string packageListPath, ILogger logger)
+		{
+			logger.LogInformation("Running UnrealVirtualizationTool...");
+
+			string toolPath = Path.Combine(engineRoot, @"Engine\Binaries\Win64\UnrealVirtualizationTool.exe");
+			string toolArgs = string.Format("-ClientSpecName={0} -Mode=PackageList -Path={1}", clientSpec, packageListPath);
+
+			using (Stream stdOutput = Console.OpenStandardOutput())
+			using (ManagedProcessGroup Group = new ManagedProcessGroup())
+			using (ManagedProcess Process = new ManagedProcess(Group, toolPath, toolArgs, null, null, System.Diagnostics.ProcessPriorityClass.Normal))
+			{
+				await Process.CopyToAsync(stdOutput, CancellationToken.None);
+
+				if (Process.ExitCode != 0)
+				{
+					logger.LogError("UnrealVirtualizationTool failed!");
+					return false;
+				}
+			}
+
+			return true;
+		}
+
+		/// <summary>
+		/// Writes out a list of package paths to a text file stored in the users 
+		/// temp directory.
+		/// The name of the file will be a randomly generated GUID making file name
+		/// collisions unlikely.
+		/// Each path will be written to it's own line inside of the files
+		/// </summary>
+		/// <param name="packagePaths">A list of package file paths to be written to the file</param>
+		/// <param name="logger">Interface for logging</param>
+		/// <returns>The path of the file once it has been written</returns>
+		private static async Task<string> WritePackageFileList(List<string> packagePaths, ILogger logger)
+		{
+			// We pass the list of packages to the tool via a file as the number of package 
+			// paths can potentially be huge and exceed the cmdline length. 
+			// So currently we write the files to a UnrealVirtualizationTool directory under the temp directory
+			string tempDirectory = Path.Combine(Path.GetTempPath(), "UnrealVirtualizationTool");
+			Directory.CreateDirectory(tempDirectory);
+
+			string tempFilesPath = Path.Combine(tempDirectory, Guid.NewGuid().ToString() + ".txt");
+
+			logger.LogInformation("Writing package path list to '{Path}'...", tempFilesPath);
+
+			using (StreamWriter Writer = new StreamWriter(tempFilesPath))
+			{
+				foreach (string line in packagePaths)
+				{
+					await Writer.WriteLineAsync(line);
+				}
+			}
+
+			return NormalizeFilename(tempFilesPath);
+		}
+
+		/// <summary>
+		/// Sorts the given package paths by the unreal project that they are found to be in.
+		/// If a package is found to not be in a project it will currently raise a warning 
+		/// but not prevent further execution.
+		/// </summary>
+		/// <param name="packagePaths">A list of absolute file paths pointing to package files</param>
+		/// <param name="logger">Interface for logging</param>
+		/// <returns>A dictionary where the key is the path of an unreal project and the value is a list of packages in that project</returns>
+		private static Dictionary<string, List<string>> SortPackagesIntoProjects(string[] packagePaths, ILogger logger)
+		{
+			Dictionary<string, List<string>> projects = new Dictionary<string, List<string>>();
+			foreach (string path in packagePaths)
+			{
+				string normalizedPath = NormalizeFilename(path);
+				string projectFilePath = FindProjectForPackage(normalizedPath);
+
+				if (!String.IsNullOrEmpty(projectFilePath))
+				{
+					if (!projects.ContainsKey(projectFilePath))
+					{
+						projects.Add(projectFilePath, new List<string>());
+					}
+
+					projects[projectFilePath].Add(normalizedPath);
+				}
+				else
+				{
+					// @todo Re-evaluate if we want this warning at some point in the future. 
+					// Technically submitting a package file not under a project could have valid use cases
+					// but if the user called this command it would indicate that they expect that they are
+					// submitting packages in a valid project.
+					// So for now we give a warning so that it is easier to catch cases where packages are
+					// thought not to be part of a project even if they are.
+					logger.LogWarning("Unable to find a valid project for the package '{Path}'", normalizedPath);
+				}
+			}
+
+			return projects;
+		}
+
+		/// <summary>
+		/// Adds the virtualization tags to a changelist description which is used to show that the virtualization
+		/// process has been run on it, before it was submitted.
+		/// </summary>
+		/// <param name="perforceConnection">A valid connection to perforce to use</param>
+		/// <param name="changeNumber">The changelist we should be stamping</param>
+		/// <param name="logger">Interface for logging</param>
+		/// <returns>If the stamp succeeded then it returns a ChangeRecord representing the changelist as it was before it was stamped, returns null on failure</returns>
+		private static async Task<ChangeRecord?> StampChangelistDescription(IPerforceConnection perforceConnection, int changeNumber, ILogger logger)
+		{
+			logger.LogInformation("Stamping changelist description with virtualization tags");
+
+			ChangeRecord changeRecord;
+			try
+			{
+				changeRecord = await perforceConnection.GetChangeAsync(GetChangeOptions.None, changeNumber);
+			}
+			catch (Exception)
+			{
+				logger.LogError("Failed to get the description {Change} so we can edit it", changeNumber);
+				return null;
+			}
+
+			string? originalDescription = changeRecord.Description;
+
+			// @todo Should be getting the tag from the virtualization tool!
+			changeRecord.Description += "\n#virtualized\n";
+
+			PerforceResponse updateResponse = await perforceConnection.TryUpdateChangeAsync(UpdateChangeOptions.None, changeRecord, CancellationToken.None);
+			if (updateResponse.Succeeded)
+			{
+				// Restore the original description so we are returning the original ChangeRecord
+				changeRecord.Description = originalDescription;
+				return changeRecord;
+			}
+			else
+			{
+				logger.LogError("Failed to edit the description of {Change} due to\n{Message}", changeNumber, updateResponse.Error!.ToString());
+				return null;
+			}
+		}
+
+		/// <summary>
+		/// Find all of the unreal packages in a single perforce changelist and return them as local file paths on the users machine.
+		/// </summary>
+		/// <param name="perforceConnection">A valid connection to perforce. This should have the correct client spec for the given changelist number</param>
+		/// <param name="changeNumber">The changelist we should look in</param>
+		/// <param name="logger">Interface for logging</param>
+		/// <returns>A list of all the packages in the given changelist, in local file path format</returns>
+		private static async Task<string[]?> FindPackagesInChangelist(IPerforceConnection perforceConnection, int changeNumber, ILogger logger)
+		{
+			logger.LogInformation("Finding files in changelist {Change}...", changeNumber);
+			DescribeRecord changeRecord;
+			try
+			{
+				changeRecord = await perforceConnection.DescribeAsync(changeNumber, CancellationToken.None);
+			}
+			catch (Exception)
+			{
+				logger.LogError("Failed to find changelist {Change}", changeNumber);
+				return null;
+			}
+
+			// @todo Should we check if the changelist has shelved files and error at this point since
+			// we know that the user will not be able to submit it?
+
+			// Find the depth paths in the changelist that point to package files
+			string[] depotPackagePaths = changeRecord.Files.Select(x => x.DepotFile).Where(x => IsPackagePath(x)).ToArray();
+
+			// Now convert from depot paths to local paths on the users machine
+			List<WhereRecord> whereRecords = await perforceConnection.WhereAsync(depotPackagePaths, CancellationToken.None).ToListAsync();
+
+			return whereRecords.Select(x => x.Path).ToArray();
+		}
+
+		/// <summary>
+		/// Finds the unreal project file for a given unreal package
+		/// </summary>
+		/// <param name="packagePath">The package to find the project for</param>
+		/// <returns>The path of the projects .uproject file if found, an empty string if no valid project file was found</returns>
+		private static string FindProjectForPackage(string packagePath)
+		{
+			// @todo note that this function mirrors FUnrealVirtualizationToolApp::TryFindProject
+			// both will not work correctly with some plugin setups, this is known and will be
+			// fixed later when FUnrealVirtualizationToolApp is also fixed.
+			packagePath = NormalizeFilename(packagePath);
+
+			int contentIndex = packagePath.IndexOf("/content/", StringComparison.OrdinalIgnoreCase);
+			if (contentIndex == -1)
+			{
+				// Error condition no content directory?
+				return string.Empty;
+			}
+
+			string projectDirectory = packagePath[..contentIndex];
+
+			string[] projectFiles = Directory.GetFiles(projectDirectory, "*.uproject");
+
+			if (projectFiles.Length == 0)
+			{
+				// If there was no project file, the package could be in a plugin, so lets check for that
+				int pluginIndex = packagePath.IndexOf("/plugins/", StringComparison.OrdinalIgnoreCase);
+				if (pluginIndex == -1)
+				{
+					// Error condition not a plugin?
+					return string.Empty;
+				}
+
+				projectDirectory = packagePath[..pluginIndex];
+				projectFiles = Directory.GetFiles(projectDirectory, "*.uproject");
+			}
+
+			if (projectFiles.Length == 1)
+			{
+				return NormalizeFilename(projectFiles[0]);
+			}
+			else if (projectFiles.Length > 1)
+			{
+				// Error condition too many project files?
+				return string.Empty;
+			}
+			else
+			{
+				// Error condition no project file found
+				return string.Empty;
+			}
+		}
+
+		// The following functions mirror code found @Engine\Source\Runtime\CoreUObject\Private\Misc\PackageName.cpp
+		#region CoreUObject PackageName 
+
+		private static bool IsPackagePath(string path)
+		{
+			return IsBinaryPackagePath(path) || IsTextPackagePath(path);
+		}
+		
+		private static bool IsBinaryPackagePath(string path)
+		{
+			return path.EndsWith(".uasset", StringComparison.OrdinalIgnoreCase) || path.EndsWith(".umap", StringComparison.OrdinalIgnoreCase);
+		}
+
+		private static bool IsTextPackagePath(string path)
+		{
+			return path.EndsWith(".utxt", StringComparison.OrdinalIgnoreCase) || path.EndsWith(".utxtmap", StringComparison.OrdinalIgnoreCase);
+		}
+
+		#endregion
+
+		// The following functions mirror code found @Engine\Source\Runtime\Core\Private\Misc\Paths.cpp
+		#region Paths
+
+		private static string NormalizeFilename(string path)
+		{
+			return path.Replace('\\', '/');
+		}
+
+		private static string NormalizeDirectoryName(string path)
+		{
+			path = path.Replace('\\', '/');
+
+			if (path.EndsWith('/'))
+			{
+				path = path.Remove(path.Length - 1, 1);
+			}
+
+			return path;
+		}
+
+		#endregion
+
+		// The following functions mirror code found @Engine\Source\Developer\DesktopPlatform\Private\DesktopPlatformBase.cpp
+		#region DesktopPlatform DesktopPlatformBase
+
+		// Note that for the C# versions we only support the modern way of associating
+		// projects with engine installation, via the 'EngineAssociation' entry in
+		// the .uproject file.
+		private static string GetEngineIdentifierForProject(string projectFilePath)
+		{
+			try
+			{
+				JsonObject root = JsonObject.Read(new FileReference(projectFilePath));
+				string engineIdentifier = root.GetStringField("EngineAssociation");
+				if (!String.IsNullOrEmpty(engineIdentifier))
+				{
+					// @todo what if it is a path?
+					return engineIdentifier;
+				}
+			}
+			catch (Exception)
+			{
+				// Unable to parse the project file correctly
+				return string.Empty;
+			}
+
+			// The project file did not contain the identifier, so we need to search the
+			// directory hierarchy to try and fine the engine
+			DirectoryInfo directoryToSearch = new DirectoryInfo(Path.GetDirectoryName(projectFilePath));
+
+			while(directoryToSearch != null)
+			{
+				if (IsValidRootDirectory(directoryToSearch.ToString()))
+				{
+					return GetEngineIdentifierFromRootDir(directoryToSearch.ToString());
+				}
+
+				directoryToSearch = Directory.GetParent(directoryToSearch.ToString());
+			}
+
+			return string.Empty;
+		}
+
+		private static string GetEngineRootDirFromIdentifier(string engineIdentifier)
+		{
+			// @todo remove
+			using ILoggerFactory Factory = LoggerFactory.Create(Builder => Builder.AddEpicDefault());
+			ILogger Logger = Factory.CreateLogger<Program>();
+
+			IReadOnlyDictionary<string, string> engineInstalls = EnumerateEngineInstallations(Logger);
+
+			if (engineInstalls.TryGetValue(engineIdentifier, out string? engineRoot))
+			{
+				return engineRoot;
+			}
+			else
+			{
+				return string.Empty;
+			}
+		}
+
+		private static string GetEngineIdentifierFromRootDir(string rootDirectory)
+		{
+			rootDirectory = NormalizeDirectoryName(rootDirectory);
+			
+			// @todo remove
+			using ILoggerFactory Factory = LoggerFactory.Create(Builder => Builder.AddEpicDefault());
+			ILogger Logger = Factory.CreateLogger<Program>();
+
+			IReadOnlyDictionary<string, string> engineInstalls = EnumerateEngineInstallations(Logger);
+
+			foreach (KeyValuePair<string, string> pair in engineInstalls)
+			{
+				if (String.Equals(pair.Value,rootDirectory, StringComparison.OrdinalIgnoreCase))
+				{
+					return pair.Key;
+				}
+			}
+
+			return String.Empty;
+		}
+
+		private static bool IsValidRootDirectory(string rootDirectory)
+		{
+			// Check that there's an Engine\Binaries directory underneath the root
+			string engineBinariesDirectory = Path.Combine(rootDirectory, "Engine/Binaries");
+
+			if (!Directory.Exists(engineBinariesDirectory))
+			{
+				return false;
+			}
+
+			// Also check there's an Engine\Build directory. This will filter out anything
+			// //that has an engine-like directory structure but doesn't allow building
+			// code projects - like the launcher.
+
+			string engineBuildDirectory = Path.Combine(rootDirectory, "Engine/Build");
+
+			if (!Directory.Exists(engineBuildDirectory))
+			{
+				return false;
+			}
+
+			return true;
+		}
+
+		private static IReadOnlyDictionary<string, string> EnumerateEngineInstallations(ILogger logger)
+		{
+			Dictionary<string, string> launcherEngineInstalls = EnumerateLauncherEngineInstallations(logger);
+			Dictionary<string, string> customEngineInstalls = EnumerateCustomEngineInstallations(logger);
+
+			Dictionary<string, string> allInstalls = new Dictionary<string, string>();
+
+			foreach (KeyValuePair<string, string> pair in launcherEngineInstalls)
+			{
+				allInstalls.Add(pair.Key, pair.Value);
+			}
+
+			foreach (KeyValuePair<string, string> pair in customEngineInstalls)
+			{
+				allInstalls.Add(pair.Key, pair.Value);
+			}
+
+			return allInstalls;
+		}
+
+		private static Dictionary<string, string> EnumerateLauncherEngineInstallations(ILogger logger)
+		{
+			Dictionary<string, string> installations = new Dictionary<string, string>();
+
+			try
+			{
+				string installedListFilePath = Path.Combine(Environment.GetFolderPath(Environment.SpecialFolder.CommonApplicationData), "Epic/UnrealEngineLauncher/LauncherInstalled.dat");
+
+				if (JsonObject.TryRead(new FileReference(installedListFilePath), out JsonObject? rawObject))
+				{
+					JsonObject[] installationArray = rawObject.GetObjectArrayField("InstallationList");
+					foreach (JsonObject installObject in installationArray)
+					{
+						string appName = installObject.GetStringField("AppName");
+						if (appName.StartsWith("UE_", StringComparison.Ordinal))
+						{
+							appName = appName.Remove(0, 3);
+							string installPath = NormalizeDirectoryName(installObject.GetStringField("InstallLocation"));
+
+							installations.Add(appName, installPath);
+						}
+					}
+				}
+			}
+			catch (Exception ex)
+			{
+				logger.LogWarning("EnumerateLauncherEngineInstallations: {Message}", ex.Message);
+				installations.Clear();
+			}
+
+			return installations;
+		}
+
+		private static Dictionary<string, string> EnumerateCustomEngineInstallations(ILogger logger)
+		{
+			Dictionary<string, string> installations = new Dictionary<string, string>();
+
+			try
+			{
+				using RegistryKey? subKey = Registry.CurrentUser.OpenSubKey("SOFTWARE\\Epic Games\\Unreal Engine\\Builds", false);
+				if (subKey != null)
+				{
+					foreach (string installName in subKey.GetValueNames())
+					{
+						string? installPath = subKey.GetValue(installName) as string;
+						if (!String.IsNullOrEmpty(installPath))
+						{
+							installations.Add(installName, NormalizeDirectoryName(installPath));
+						}
+					}
+				}
+			}
+			catch (Exception ex)
+			{
+				logger.LogWarning("EnumerateCustomEngineInstallations: {Message}", ex.Message);
+				installations.Clear();
+			}
+
+			return installations;
+		}
+
+		#endregion
+
+		// The following code acts as an extension to code found @Engine\Source\Programs\Shared\EpicGames.Perforce\PerforceConnection.cs
+		#region PerforceConnection
+
+		// @todo None of the submit functions in PerforceConnection.cs seem to report submit errors in a way that we can inform the user
+		// so this is a custom version that returns the entire response list as different submit errors will return different error
+		// responses. We should fix the API in PerforceConnection so that other areas of code can get better submit error reporting.
+		private static async Task<PerforceResponseList<SubmitRecord>> TrySubmitAsync(IPerforceConnection connection, int changeNumber, SubmitOptions options, CancellationToken cancellationToken = default)
+		{
+			List<string> arguments = new List<string>();
+			if ((options & SubmitOptions.ReopenAsEdit) != 0)
+			{
+				arguments.Add("-r");
+			}
+			arguments.Add($"-c{changeNumber}");
+
+			return (await connection.CommandAsync<SubmitRecord>("submit", arguments, null, cancellationToken));
+		}
+
+		#endregion
+	}
+}
diff --git a/Engine/Extras/P4VUtils/Program.cs b/Engine/Extras/P4VUtils/Program.cs
--- a/Engine/Extras/P4VUtils/Program.cs
+++ b/Engine/Extras/P4VUtils/Program.cs
@@ -51,6 +51,12 @@ namespace P4VUtils
 			["copyclnum"] = new CopyCLCommand(),
 		};
 
+		// UESubmit - commands that help with submitting files/changelists
+		public static IReadOnlyDictionary<string, Command> SubmissionCommands { get; } = new Dictionary<string, Command>(StringComparer.OrdinalIgnoreCase)
+		{
+			["submitandvirtualize"] = new SubmitAndVirtualizeCommand(),
+		};
+
 		// UEHelpers - commands that help with common but simple operations
 		public static IReadOnlyDictionary<string, Command> HelperCommands { get; } = new Dictionary<string, Command>(StringComparer.OrdinalIgnoreCase)
 		{
@@ -83,7 +89,7 @@ namespace P4VUtils
 			["movewriteablepreflightandsubmit"] = new MoveWriteableFilesthenPreflightAndSubmitCommand(),
 		};
 
-		public static IDictionary<string, Command> Commands = RootHelperCommands.Concat(HelperCommands).Concat(IntegrateCommands).Concat(HordeCommands).ToDictionary(p => p.Key, p => p.Value, StringComparer.OrdinalIgnoreCase);
+		public static IDictionary<string, Command> Commands = SubmissionCommands.Concat(RootHelperCommands).Concat(HelperCommands).Concat(IntegrateCommands).Concat(HordeCommands).ToDictionary(p => p.Key, p => p.Value, StringComparer.OrdinalIgnoreCase);
 
 		static void PrintHelp(ILogger Logger)
 		{
@@ -420,6 +426,7 @@ namespace P4VUtils
 			if (bInstall)
 			{
 				InstallCommandsListInFolder("UERootHelpers", false/*AddFolderToContextMenu*/, RootHelperCommands, Document, DotNetLocation, AssemblyLocation, Logger);
+				InstallCommandsListInFolder("UESubmit", true/*AddFolderToContextMenu*/, SubmissionCommands, Document, DotNetLocation, AssemblyLocation, Logger);
 				InstallCommandsListInFolder("UEHelpers", true/*AddFolderToContextMenu*/, HelperCommands, Document, DotNetLocation, AssemblyLocation, Logger);
 				InstallCommandsListInFolder("UEIntegrate", true/*AddFolderToContextMenu*/, IntegrateCommands, Document, DotNetLocation, AssemblyLocation, Logger);
 				InstallCommandsListInFolder("UEHorde", true/*AddFolderToContextMenu*/, HordeCommands, Document, DotNetLocation, AssemblyLocation, Logger);
