Fix the "UE Submit -> Virtualize and Submit" command to work when no engine is registered but could still be found via the directory hierarchy

#rnx

### Misc
- Failing to find an associated engine for a project will now cause an error message to be logged and cause the changelist not to be submitted rather than allowing the process to continue.
- We now pass the logger into the methods in the DesktopPlatformBase region, rather than using a dummy logger (this was always the intent and something I missed fixing up before the original submit)

### Bug Fix
- Instead of finding the engine identifier and then finding the engine root (which mirrors native code) a new method has been added ::GetEngineIdentifierForProject which just returns the engine root directory which is all we are interested in.
- The main failing with the previous approach is that the engine association is not always registered in the system registry. I suspect that an engine install taken from p4 will only register if it is used to create a new project, meaning that the average developer will not have this info avaliable. ::GetEngineIdentifierForProject takes this into account and tries to find the engine via the directory structure before we look in the registry.
- ::GetEngineIdentifierForProject no longer tries to find the engine identifier if it could not be found in the project path and so does not properly match the native version of the method. To match the original it would need to register the engine if found, in which case the virtualization system (when run via the editor) should probably take steps to register the engine as well. This should be done as it's own work item to make sure it is tested properly.

diff --git a/Engine/Extras/P4VUtils/Commands/SubmitCommand.cs b/Engine/Extras/P4VUtils/Commands/SubmitCommand.cs
--- a/Engine/Extras/P4VUtils/Commands/SubmitCommand.cs
+++ b/Engine/Extras/P4VUtils/Commands/SubmitCommand.cs
@@ -74,37 +74,37 @@ namespace P4VUtils.Commands
 
 			foreach(KeyValuePair<string, List<string>> project in projects)
 			{
-				// Get engine association from project file
-				string engineIdentifier = GetEngineIdentifierForProject(project.Key);
-				if (!String.IsNullOrEmpty(engineIdentifier))
+				string engineRoot = GetEngineRootForProject(project.Key, logger);
+				if (!String.IsNullOrEmpty(engineRoot))
 				{
-					string engineRoot = GetEngineRootDirFromIdentifier(engineIdentifier);
-					if (!String.IsNullOrEmpty(engineRoot))
-					{
-						logger.LogInformation("\nAttempting to virtualize packages for project '{Project}' with the engine installation '{Engine}'", project.Key, engineRoot);
-						// @todo Many projects can share the same engine install, and technically UnrealVirtualizationTool
-						// supports the virtualization files from many projects at the same time. We could consider doing
-						// this pass per engine install rather than per project? At the very least we should only 'build'
-						// the tool once per engine
+					logger.LogInformation("\nAttempting to virtualize packages for project '{Project}' with the engine installation '{Engine}'", project.Key, engineRoot);
+					// @todo Many projects can share the same engine install, and technically UnrealVirtualizationTool
+					// supports the virtualization files from many projects at the same time. We could consider doing
+					// this pass per engine install rather than per project? At the very least we should only 'build'
+					// the tool once per engine
 
-						Task<bool> compileResult = BuildVirtualizationTool(engineRoot, logger);
+					Task<bool> compileResult = BuildVirtualizationTool(engineRoot, logger);
 
-						string tempFilesPath = await WritePackageFileList(project.Value, logger);
+					string tempFilesPath = await WritePackageFileList(project.Value, logger);
 
-						// Check if the compilation of the tool succeeded or not
-						if (await compileResult == false)
-						{
-							return 1;
-						}
+					// Check if the compilation of the tool succeeded or not
+					if (await compileResult == false)
+					{
+						return 1;
+					}
 
-						if(await RunVirtualizationTool(engineRoot, clientSpec, tempFilesPath, logger) == false)
-						{
-							return 1;
-						};
+					if(await RunVirtualizationTool(engineRoot, clientSpec, tempFilesPath, logger) == false)
+					{
+						return 1;
+					};
 
-						// @todo Make sure this always gets cleaned up when we go out of scope
-						File.Delete(tempFilesPath);
-					}
+					// @todo Make sure this always gets cleaned up when we go out of scope
+					File.Delete(tempFilesPath);
+				}
+				else
+				{
+					logger.LogError("Failed to find engine root for project {Project}", project.Key);
+					return 1;
 				}
 			}
 
@@ -433,6 +433,25 @@ namespace P4VUtils.Commands
 			}
 		}
 
+		private static string GetEngineRootForProject(string projectFilePath, ILogger logger)
+		{
+			string engineIdentifier = GetEngineIdentifierForProject(projectFilePath, logger);
+			if (!String.IsNullOrEmpty(engineIdentifier))
+			{
+				string engineRoot = GetEngineRootDirFromIdentifier(engineIdentifier, logger);
+				if (!String.IsNullOrEmpty(engineIdentifier))
+				{
+					return engineRoot;
+				}
+				else 
+				{
+					logger.LogWarning("Unable to find an engine root for installation {Identifier}, will attempt to find the engine via the directory hierarchy", engineIdentifier);
+				}	
+			}
+
+			return FindEngineFromPath(projectFilePath);
+		}
+
 		// The following functions mirror code found @Engine\Source\Runtime\CoreUObject\Private\Misc\PackageName.cpp
 		#region CoreUObject PackageName 
 
@@ -481,7 +500,8 @@ namespace P4VUtils.Commands
 		// Note that for the C# versions we only support the modern way of associating
 		// projects with engine installation, via the 'EngineAssociation' entry in
 		// the .uproject file.
-		private static string GetEngineIdentifierForProject(string projectFilePath)
+
+		private static string GetEngineIdentifierForProject(string projectFilePath, ILogger logger)
 		{
 			try
 			{
@@ -489,25 +509,32 @@ namespace P4VUtils.Commands
 				string engineIdentifier = root.GetStringField("EngineAssociation");
 				if (!String.IsNullOrEmpty(engineIdentifier))
 				{
-					// @todo what if it is a path?
+					// @todo what if it is a path? (native code has support for this
+					// possibly for an older version)
 					return engineIdentifier;
 				}
 			}
-			catch (Exception)
+			catch (Exception ex)
 			{
-				// Unable to parse the project file correctly
-				return string.Empty;
+				logger.LogError("Failed to parse {File} to find the engine association due to: {Reason}", projectFilePath, ex.Message);
 			}
 
-			// The project file did not contain the identifier, so we need to search the
-			// directory hierarchy to try and fine the engine
-			DirectoryInfo directoryToSearch = new DirectoryInfo(Path.GetDirectoryName(projectFilePath));
+			// @todo In the native version if there is no identifier we will try to
+			// find the engine root in the directory hierarchy then either find it
+			// identifier or register one if needed.
 
-			while(directoryToSearch != null)
+			return string.Empty;
+		}
+
+		private static string FindEngineFromPath(string path)
+		{
+			DirectoryInfo directoryToSearch = new DirectoryInfo(Path.GetDirectoryName(path));
+
+			while (directoryToSearch != null)
 			{
 				if (IsValidRootDirectory(directoryToSearch.ToString()))
 				{
-					return GetEngineIdentifierFromRootDir(directoryToSearch.ToString());
+					return NormalizeDirectoryName(directoryToSearch.ToString());
 				}
 
 				directoryToSearch = Directory.GetParent(directoryToSearch.ToString());
@@ -516,13 +543,9 @@ namespace P4VUtils.Commands
 			return string.Empty;
 		}
 
-		private static string GetEngineRootDirFromIdentifier(string engineIdentifier)
+		private static string GetEngineRootDirFromIdentifier(string engineIdentifier, ILogger logger)
 		{
-			// @todo remove
-			using ILoggerFactory Factory = LoggerFactory.Create(Builder => Builder.AddEpicDefault());
-			ILogger Logger = Factory.CreateLogger<Program>();
-
-			IReadOnlyDictionary<string, string> engineInstalls = EnumerateEngineInstallations(Logger);
+			IReadOnlyDictionary<string, string> engineInstalls = EnumerateEngineInstallations(logger);
 
 			if (engineInstalls.TryGetValue(engineIdentifier, out string? engineRoot))
 			{
@@ -534,15 +557,11 @@ namespace P4VUtils.Commands
 			}
 		}
 
-		private static string GetEngineIdentifierFromRootDir(string rootDirectory)
+		private static string GetEngineIdentifierFromRootDir(string rootDirectory, ILogger logger)
 		{
 			rootDirectory = NormalizeDirectoryName(rootDirectory);
 			
-			// @todo remove
-			using ILoggerFactory Factory = LoggerFactory.Create(Builder => Builder.AddEpicDefault());
-			ILogger Logger = Factory.CreateLogger<Program>();
-
-			IReadOnlyDictionary<string, string> engineInstalls = EnumerateEngineInstallations(Logger);
+			IReadOnlyDictionary<string, string> engineInstalls = EnumerateEngineInstallations(logger);
 
 			foreach (KeyValuePair<string, string> pair in engineInstalls)
 			{
