Clean up the logging for the P4VUtils Submit and Virtualize command

#rnx

diff --git a/Engine/Extras/P4VUtils/Commands/SubmitCommand.cs b/Engine/Extras/P4VUtils/Commands/SubmitCommand.cs
--- a/Engine/Extras/P4VUtils/Commands/SubmitCommand.cs
+++ b/Engine/Extras/P4VUtils/Commands/SubmitCommand.cs
@@ -77,7 +77,7 @@ namespace P4VUtils.Commands
 				string engineRoot = GetEngineRootForProject(project.Key, logger);
 				if (!String.IsNullOrEmpty(engineRoot))
 				{
-					logger.LogInformation("\nAttempting to virtualize packages for project '{Project}' with the engine installation '{Engine}'", project.Key, engineRoot);
+					logger.LogInformation("\nAttempting to virtualize packages in project '{Project}' via the engine installation '{Engine}'", project.Key, engineRoot);
 					// @todo Many projects can share the same engine install, and technically UnrealVirtualizationTool
 					// supports the virtualization files from many projects at the same time. We could consider doing
 					// this pass per engine install rather than per project? At the very least we should only 'build'
@@ -93,7 +93,13 @@ namespace P4VUtils.Commands
 						return 1;
 					}
 
-					if(await RunVirtualizationTool(engineRoot, clientSpec, tempFilesPath, logger) == false)
+					// Even though this will have been done while we were waiting for BuildVirtualizationTool to complete 
+					// we want to log that it was done after so that the output log makes sense to the user, otherwise 
+					// they will end up thinking that they are waiting on the PackageList to be written rather than on
+					// the tool to be built.
+					logger.LogInformation("PackageList was written to '{Path}'", tempFilesPath);
+
+					if (await RunVirtualizationTool(engineRoot, clientSpec, tempFilesPath, logger) == false)
 					{
 						return 1;
 					};
@@ -108,13 +114,15 @@ namespace P4VUtils.Commands
 				}
 			}
 
+			logger.LogInformation("\nAll packages have been virtualized");
+
 			ChangeRecord? changeRecord = await StampChangelistDescription(perforceConnection, changeNumber, logger);
 			if (changeRecord == null)
 			{
 				return 1;
 			}
 
-			logger.LogInformation("Virtualization completed, attempting to submit changelist {Number}", changeNumber);
+			logger.LogInformation("Attempting to submit changelist {Number}...", changeNumber);
 
 			// Submit
 			//@todo ideally we should get the tags back from UnrealVirtualizationTool
@@ -142,7 +150,7 @@ namespace P4VUtils.Commands
 				// If the final submit failed we remove the virtualization tag, even though the changelist is technically
 				// virtualized at this point and submitting it would be safe.
 				// This is to keep the behavior the same as the native code paths.
-				logger.LogInformation("Removing virtualization tags from the changelist..");
+				logger.LogInformation("Removing virtualization tags from the changelist...");
 				PerforceResponse updateResponse = await perforceConnection.TryUpdateChangeAsync(UpdateChangeOptions.None, changeRecord, CancellationToken.None);
 				if (!updateResponse.Succeeded)
 				{
@@ -252,8 +260,6 @@ namespace P4VUtils.Commands
 
 			string tempFilesPath = Path.Combine(tempDirectory, Guid.NewGuid().ToString() + ".txt");
 
-			logger.LogInformation("Writing package path list to '{Path}'...", tempFilesPath);
-
 			using (StreamWriter Writer = new StreamWriter(tempFilesPath))
 			{
 				foreach (string line in packagePaths)
@@ -315,7 +321,7 @@ namespace P4VUtils.Commands
 		/// <returns>If the stamp succeeded then it returns a ChangeRecord representing the changelist as it was before it was stamped, returns null on failure</returns>
 		private static async Task<ChangeRecord?> StampChangelistDescription(IPerforceConnection perforceConnection, int changeNumber, ILogger logger)
 		{
-			logger.LogInformation("Stamping changelist description with virtualization tags");
+			logger.LogInformation("Adding virtualization tags to changelist description...");
 
 			ChangeRecord changeRecord;
 			try
