Allow the P4VUtils submit command to deal correctly with changelists containing no package files.

#rnx

- The c# p4 api seems to have problems when calling p4 where, with no valid files so make sure that we early out before calling that.
- Only run the virtualization process if we have package files, no point trying to virtualize anything else.
-- If the changelist does not contain any packages then we just submit the change list, we don't need to stamp it or display logging info about how 0 packages are in 0 projects etc.
- Moved the code that submits the changelist to it's own method.

diff --git a/Engine/Extras/P4VUtils/Commands/SubmitCommand.cs b/Engine/Extras/P4VUtils/Commands/SubmitCommand.cs
--- a/Engine/Extras/P4VUtils/Commands/SubmitCommand.cs
+++ b/Engine/Extras/P4VUtils/Commands/SubmitCommand.cs
@@ -53,111 +53,104 @@ namespace P4VUtils.Commands
 				return 1;
 			}
 
-			// First we need to find the packages in the changelist
 			// First we need to find the packages in the changelist
 			string[]? localFilePaths = await FindPackagesInChangelist(perforceConnection, changeNumber, logger);
 			if(localFilePaths == null)
 			{
 				return 1;
 			}
-			logger.LogInformation("Found {Amount} package(s) that may need virtualization", localFilePaths.Length);
 
-			// Now sort the package paths by their project (it is unlikely but a user could be submitting content
-			// from multiple projects that the same time)
+			if (localFilePaths.Length > 0)
+			{
+				logger.LogInformation("Found {Amount} package(s) that may need virtualization", localFilePaths.Length);
 
-			Dictionary<string, List<string>> projects = SortPackagesIntoProjects(localFilePaths, logger);
+				// Now sort the package paths by their project (it is unlikely but a user could be submitting content
+				// from multiple projects that the same time)
 
-			logger.LogInformation("The packages are distributed between {Amount} project(s)", projects.Count);
+				Dictionary<string, List<string>> projects = SortPackagesIntoProjects(localFilePaths, logger);
 
-			// Find engine per project
-			IReadOnlyDictionary<string, string> engineInstalls = EnumerateEngineInstallations(logger);
+				logger.LogInformation("The packages are distributed between {Amount} project(s)", projects.Count);
 
-			foreach(KeyValuePair<string, List<string>> project in projects)
-			{
-				string engineRoot = GetEngineRootForProject(project.Key, logger);
-				if (!String.IsNullOrEmpty(engineRoot))
+				// Find engine per project
+				IReadOnlyDictionary<string, string> engineInstalls = EnumerateEngineInstallations(logger);
+
+				foreach (KeyValuePair<string, List<string>> project in projects)
 				{
-					logger.LogInformation("\nAttempting to virtualize packages in project '{Project}' via the engine installation '{Engine}'", project.Key, engineRoot);
-					// @todo Many projects can share the same engine install, and technically UnrealVirtualizationTool
-					// supports the virtualization files from many projects at the same time. We could consider doing
-					// this pass per engine install rather than per project? At the very least we should only 'build'
-					// the tool once per engine
+					string engineRoot = GetEngineRootForProject(project.Key, logger);
+					if (!String.IsNullOrEmpty(engineRoot))
+					{
+						logger.LogInformation("\nAttempting to virtualize packages in project '{Project}' via the engine installation '{Engine}'", project.Key, engineRoot);
+						// @todo Many projects can share the same engine install, and technically UnrealVirtualizationTool
+						// supports the virtualization files from many projects at the same time. We could consider doing
+						// this pass per engine install rather than per project? At the very least we should only 'build'
+						// the tool once per engine
 
-					Task<bool> compileResult = BuildVirtualizationTool(engineRoot, logger);
+						Task<bool> compileResult = BuildVirtualizationTool(engineRoot, logger);
 
-					string tempFilesPath = await WritePackageFileList(project.Value, logger);
+						string tempFilesPath = await WritePackageFileList(project.Value, logger);
 
-					// Check if the compilation of the tool succeeded or not
-					if (await compileResult == false)
-					{
-						return 1;
-					}
+						// Check if the compilation of the tool succeeded or not
+						if (await compileResult == false)
+						{
+							return 1;
+						}
+
+						// Even though this will have been done while we were waiting for BuildVirtualizationTool to complete 
+						// we want to log that it was done after so that the output log makes sense to the user, otherwise 
+						// they will end up thinking that they are waiting on the PackageList to be written rather than on
+						// the tool to be built.
+						logger.LogInformation("PackageList was written to '{Path}'", tempFilesPath);
 
-					// Even though this will have been done while we were waiting for BuildVirtualizationTool to complete 
-					// we want to log that it was done after so that the output log makes sense to the user, otherwise 
-					// they will end up thinking that they are waiting on the PackageList to be written rather than on
-					// the tool to be built.
-					logger.LogInformation("PackageList was written to '{Path}'", tempFilesPath);
+						if (await RunVirtualizationTool(engineRoot, clientSpec, tempFilesPath, logger) == false)
+						{
+							return 1;
+						};
 
-					if (await RunVirtualizationTool(engineRoot, clientSpec, tempFilesPath, logger) == false)
+						// @todo Make sure this always gets cleaned up when we go out of scope
+						File.Delete(tempFilesPath);
+					}
+					else
 					{
+						logger.LogError("Failed to find engine root for project {Project}", project.Key);
 						return 1;
-					};
-
-					// @todo Make sure this always gets cleaned up when we go out of scope
-					File.Delete(tempFilesPath);
+					}
 				}
-				else
+
+				logger.LogInformation("\nAll packages have been virtualized");
+
+				//@todo ideally we should get the tags back from UnrealVirtualizationTool
+				ChangeRecord? changeRecord = await StampChangelistDescription(perforceConnection, changeNumber, logger);
+				if (changeRecord == null)
 				{
-					logger.LogError("Failed to find engine root for project {Project}", project.Key);
 					return 1;
 				}
-			}
 
-			logger.LogInformation("\nAll packages have been virtualized");
+				logger.LogInformation("Attempting to submit changelist {Number}...", changeNumber);
 
-			ChangeRecord? changeRecord = await StampChangelistDescription(perforceConnection, changeNumber, logger);
-			if (changeRecord == null)
-			{
-				return 1;
-			}
-
-			logger.LogInformation("Attempting to submit changelist {Number}...", changeNumber);
-
-			// Submit
-			//@todo ideally we should get the tags back from UnrealVirtualizationTool
-			PerforceResponseList<SubmitRecord> submitResponses = await TrySubmitAsync(perforceConnection, changeNumber, SubmitOptions.None, CancellationToken.None);
-
-			bool successfulSubmit = submitResponses.All(x => x.Succeeded);
-			if (successfulSubmit)
-			{
-				// @todo The submit API really should return the number that the changelist was finally submit as so that
-				// we can log that here instead.
-				logger.LogInformation("Successfully submited changelist {Change}", changeNumber);
-			}
-			else
-			{
-				// Log every response that was a failure and has an error message associated with it
-				logger.LogError("Submit failed due to:");
-				foreach (PerforceResponse<SubmitRecord> response in submitResponses)
+				// Submit
+				if (await SubmitChangelist(perforceConnection, changeNumber, logger) == false)
 				{
-					if (response.Failed && response.Error != null)
+					// If the final submit failed we remove the virtualization tag, even though the changelist is technically
+					// virtualized at this point and submitting it would be safe.
+					// This is to keep the behavior the same as the native code paths.
+					logger.LogInformation("Removing virtualization tags from the changelist...");
+					PerforceResponse updateResponse = await perforceConnection.TryUpdateChangeAsync(UpdateChangeOptions.None, changeRecord, CancellationToken.None);
+					if (!updateResponse.Succeeded)
 					{
-						logger.LogError("{Message}", response.Error.ToString());
+						logger.LogError("Failed to remove the virtualization tags!");
 					}
+
+					return 1;
 				}
+			}
+			else
+			{
+				logger.LogInformation("The changelist does not contain any package files, submitting as normal...");
 
-				// If the final submit failed we remove the virtualization tag, even though the changelist is technically
-				// virtualized at this point and submitting it would be safe.
-				// This is to keep the behavior the same as the native code paths.
-				logger.LogInformation("Removing virtualization tags from the changelist...");
-				PerforceResponse updateResponse = await perforceConnection.TryUpdateChangeAsync(UpdateChangeOptions.None, changeRecord, CancellationToken.None);
-				if (!updateResponse.Succeeded)
+				if (await SubmitChangelist(perforceConnection, changeNumber, logger) == false)
 				{
-					logger.LogError("Failed to remove the virtualization tags!");
-				}		
-
-				return 1;
+					return 1;
+				}
 			}
 		
 			return 0;
@@ -311,6 +304,34 @@ namespace P4VUtils.Commands
 			return projects;
 		}
 
+		private static async Task<bool> SubmitChangelist(IPerforceConnection perforceConnection, int changeNumber, ILogger logger)
+		{
+			PerforceResponseList<SubmitRecord> submitResponses = await TrySubmitAsync(perforceConnection, changeNumber, SubmitOptions.None, CancellationToken.None);
+
+			bool successfulSubmit = submitResponses.All(x => x.Succeeded);
+			if (successfulSubmit)
+			{
+				// @todo The submit API really should return the number that the changelist was finally submit as so that
+				// we can log that here instead.
+				logger.LogInformation("Successfully submited changelist {Change}", changeNumber);
+				return true;
+			}
+			else
+			{
+				// Log every response that was a failure and has an error message associated with it
+				logger.LogError("Submit failed due to:");
+				foreach (PerforceResponse<SubmitRecord> response in submitResponses)
+				{
+					if (response.Failed && response.Error != null)
+					{
+						logger.LogError("{Message}", response.Error.ToString());
+					}
+				}
+
+				return false;
+			}
+		}
+
 		/// <summary>
 		/// Adds the virtualization tags to a changelist description which is used to show that the virtualization
 		/// process has been run on it, before it was submitted.
@@ -380,6 +401,12 @@ namespace P4VUtils.Commands
 			// Find the depth paths in the changelist that point to package files
 			string[] depotPackagePaths = changeRecord.Files.Select(x => x.DepotFile).Where(x => IsPackagePath(x)).ToArray();
 
+			// We can early out if there are no packages in the changelist
+			if (depotPackagePaths.Length == 0)
+			{
+				return depotPackagePaths;
+			}
+
 			// Now convert from depot paths to local paths on the users machine
 			List<WhereRecord> whereRecords = await perforceConnection.WhereAsync(depotPackagePaths, CancellationToken.None).ToListAsync();
 
