Small improvements to the P4VUtil Submit&Virtualize command logging

#rnx

- Only display the p4 submit error messages, doo not include the perforce generic code as it can be confusing to the end user.
- If a submit fails, print a line to the log once the virtualization tags have been removed from the original changelist description. Previously the last message the user would see ended in "..." which implied that the process was still running.

diff --git a/Engine/Extras/P4VUtils/Commands/SubmitCommand.cs b/Engine/Extras/P4VUtils/Commands/SubmitCommand.cs
--- a/Engine/Extras/P4VUtils/Commands/SubmitCommand.cs
+++ b/Engine/Extras/P4VUtils/Commands/SubmitCommand.cs
@@ -139,6 +139,7 @@ namespace P4VUtils.Commands
 					{
 						logger.LogError("Failed to remove the virtualization tags!");
 					}
+					logger.LogInformation("Virtualization tags have been removed.");
 
 					return 1;
 				}
@@ -319,12 +320,12 @@ namespace P4VUtils.Commands
 			else
 			{
 				// Log every response that was a failure and has an error message associated with it
-				logger.LogError("Submit failed due to:");
+				logger.LogError("ERROR - Submit failed due to:");
 				foreach (PerforceResponse<SubmitRecord> response in submitResponses)
 				{
 					if (response.Failed && response.Error != null)
 					{
-						logger.LogError("{Message}", response.Error.ToString());
+						logger.LogError("\t{Message}", response.Error.Data.ToString());
 					}
 				}
 
