The VA tool will now give the user an error if they attempt to submit an empty changelist or a changelist containing a shelved file.

#rnx

### No files
- When submitting a change list with no workspace files (either empty changelist or only shelved files) our internal C# api does not return any errors.
- To get around this for now ::FindPackagesInChangelist can be changed to return null if we detect that there are no files to submit.
- I am following up with the original author to try and find out if the submit behavior in our API is intended or not.

### Shelved files
- Attempting to submit a valid changelist that has shelved files will eventually fail when we try to finally submit the changelist.
- This means that the tool will check all of the packages for virtualization first, before eventually failing on the submit.
- Most users will alt tab out of p4v during the submit and we cannot do anything fancy with the error reporting (like make the task bar icon flash etc) so might not notice that the submit failed for some time.
- So it makes sense for us to check for shelved files, before we kick off the virtualization process and if the changelist does have shelved files we can error out and the user should see it before they alt tab away.

diff --git a/Engine/Extras/P4VUtils/Commands/SubmitCommand.cs b/Engine/Extras/P4VUtils/Commands/SubmitCommand.cs
--- a/Engine/Extras/P4VUtils/Commands/SubmitCommand.cs
+++ b/Engine/Extras/P4VUtils/Commands/SubmitCommand.cs
@@ -55,13 +55,23 @@ namespace P4VUtils.Commands
 
 			// First we need to find the packages in the changelist
 			string[]? localFilePaths = await FindPackagesInChangelist(perforceConnection, changeNumber, logger);
-			if(localFilePaths == null)
+			if (localFilePaths == null)
 			{
 				return 1;
 			}
 
 			if (localFilePaths.Length > 0)
 			{
+				// If the changelist has shelved files then the submit will fail, so if we have package files that
+				// need to be checked for virtualization we don't want to have the potentially long virtualization
+				// process run (during which the user will alt-tab) only for it to fail at the end.
+				// So we should check for shelved files now and early out before the virtualization process runs.
+				if (await DoesChangelistHaveShelvedFiles(perforceConnection, changeNumber) == true)
+				{
+					logger.LogError("Changelist {Change} has shelved files and cannot be submitted", changeNumber);
+					return 1;
+				}
+
 				logger.LogInformation("Found {Amount} package(s) that may need virtualization", localFilePaths.Length);
 
 				// Now sort the package paths by their project (it is unlikely but a user could be submitting content
@@ -396,6 +406,12 @@ namespace P4VUtils.Commands
 				return null;
 			}
 
+			if (changeRecord.Files.Count == 0)
+			{
+				logger.LogError("Changelist {Change} is empty, cannot submit", changeNumber);
+				return null;
+			}
+
 			// @todo Should we check if the changelist has shelved files and error at this point since
 			// we know that the user will not be able to submit it?
 
@@ -414,6 +430,26 @@ namespace P4VUtils.Commands
 			return whereRecords.Select(x => x.Path).ToArray();
 		}
 
+		/// <summary>
+		/// Returns if a changelist has shelved files or not
+		/// </summary>
+		/// <param name="perforceConnection">A valid connection to perforce. This should have the correct client spec for the given changelist number</param>
+		/// <param name="changeNumber">The changelist we should look in</param>
+		/// <returns>True if the changelist contains shelved files, otherwise false</returns>
+		private static async Task<bool> DoesChangelistHaveShelvedFiles(IPerforceConnection perforceConnection, int changeNumber)
+		{
+			List<DescribeRecord> responses = await perforceConnection.DescribeAsync(DescribeOptions.Shelved, -1, new int[] { changeNumber }, CancellationToken.None);
+
+			if (responses.Count == 1)
+			{
+				return responses[0].Files.Count != 0;
+			}
+			else
+			{
+				return false;
+			}
+		}
+
 		/// <summary>
 		/// Finds the unreal project file for a given unreal package
 		/// </summary>
