- Some p4vutils code changes to match the already checked ion binaries

diff --git a/Engine/Extras/P4VUtils/Commands/PreflightCommand.cs b/Engine/Extras/P4VUtils/Commands/PreflightCommand.cs
--- a/Engine/Extras/P4VUtils/Commands/PreflightCommand.cs
+++ b/Engine/Extras/P4VUtils/Commands/PreflightCommand.cs
@@ -142,7 +142,7 @@ namespace P4VUtils.Commands
 						SetupVisuals();
 
 						DialogResult result= MessageBox.Show(Prompt, Caption, MessageBoxButtons.YesNoCancel);
-						Response = DialogResult.ToString();
+						Response = result.ToString();
 #endif
 					}
 					else if (OperatingSystem.IsMacOS())
diff --git a/Engine/Extras/P4VUtils/Program.cs b/Engine/Extras/P4VUtils/Program.cs
--- a/Engine/Extras/P4VUtils/Program.cs
+++ b/Engine/Extras/P4VUtils/Program.cs
@@ -48,7 +48,9 @@ namespace P4VUtils
 		public static IReadOnlyDictionary<string, Command> RootHelperCommands { get; } = new Dictionary<string, Command>(StringComparer.OrdinalIgnoreCase)
 		{
 			["describe"] = new DescribeCommand(),
+#if !IS_LINUX
 			["copyclnum"] = new CopyCLCommand(),
+#endif
 		};
 
 		// UESubmit - commands that help with submitting files/changelists
@@ -295,7 +297,7 @@ namespace P4VUtils
 							Definition.AppendChild(Description);
 
 							XmlElement Command = Document.CreateElement("Command");
-							Command.InnerText = DotNetLocation.FullName;
+							Command.InnerText = DotNetLocation.FullName.QuoteArgument();
 							Definition.AppendChild(Command);
 
 							XmlElement Arguments = Document.CreateElement("Arguments");
@@ -414,16 +416,8 @@ namespace P4VUtils
 			}
 
 
-			FileReference DotNetLocation;
-			if (OperatingSystem.IsWindows())
-			{
-				DotNetLocation = FileReference.Combine(DirectoryReference.GetSpecialFolder(Environment.SpecialFolder.ProgramFiles)!, "dotnet", "dotnet.exe");
-			}
-			else
-			{
-				// so we don't need to run SetupDotnet.sh every time we run a tool, point to the executing dotnet, assuming it will be usable later
-				DotNetLocation = new FileReference(System.Diagnostics.Process.GetCurrentProcess().MainModule.FileName);
-			}
+			// so we don't need to run SetupDotnet.sh every time we run a tool, point to the executing dotnet, assuming it will be usable later
+			FileReference DotNetLocation = new FileReference(Environment.ProcessPath!);
 			FileReference AssemblyLocation = new FileReference(Assembly.GetExecutingAssembly().GetOriginalLocation());
 
 			XmlElement? Root = Document.SelectSingleNode("CustomToolDefList") as XmlElement;
