Fix P4Utils to compile with the changes to the logging API in CL 20245803

#rnx

diff --git a/Engine/Extras/P4VUtils/Program.cs b/Engine/Extras/P4VUtils/Program.cs
--- a/Engine/Extras/P4VUtils/Program.cs
+++ b/Engine/Extras/P4VUtils/Program.cs
@@ -116,7 +116,7 @@ namespace P4VUtils
 		{
 			using ILoggerFactory Factory = LoggerFactory.Create(Builder => Builder.AddEpicDefault());//.AddSimpleConsole(Options => { Options.SingleLine = true; Options.IncludeScopes = false; }));
 			ILogger Logger = Factory.CreateLogger<Program>();
-			Log.Logger = Logger;
+			Log.SetInnerLogger(Logger);
 
 			try
 			{
