Fix some warnings introduced in the recent upgrade to .NET 6

#rnx

- Handle Path.GetDirectoryName and Directory.GetParent potentially returning null

diff --git a/Engine/Extras/P4VUtils/Commands/SubmitCommand.cs b/Engine/Extras/P4VUtils/Commands/SubmitCommand.cs
--- a/Engine/Extras/P4VUtils/Commands/SubmitCommand.cs
+++ b/Engine/Extras/P4VUtils/Commands/SubmitCommand.cs
@@ -598,16 +598,21 @@ namespace P4VUtils.Commands
 
 		private static string FindEngineFromPath(string path)
 		{
-			DirectoryInfo directoryToSearch = new DirectoryInfo(Path.GetDirectoryName(path));
-
-			while (directoryToSearch != null)
+			string? directoryName = Path.GetDirectoryName(path);
+			
+			if (!String.IsNullOrEmpty(directoryName))
 			{
-				if (IsValidRootDirectory(directoryToSearch.ToString()))
+				DirectoryInfo? directoryToSearch = new DirectoryInfo(directoryName);
+
+				while (directoryToSearch != null)
 				{
-					return NormalizeDirectoryName(directoryToSearch.ToString());
-				}
+					if (IsValidRootDirectory(directoryToSearch.ToString()))
+					{
+						return NormalizeDirectoryName(directoryToSearch.ToString());
+					}
 
-				directoryToSearch = Directory.GetParent(directoryToSearch.ToString());
+					directoryToSearch = Directory.GetParent(directoryToSearch.ToString());
+				}
 			}
 
 			return string.Empty;
