Fix more compilation warnings in P4VUtils.

#rnx

- Added checks when parsing xml to make sure that expected nodes/lists were found and are not null.
- Response was assigned as null, but nothing was checking for this. The initial value has been changed to an empty string which wioll invoke the user cancelling the operation as the default behaviour.
- StreamReader.ReadToEnd is a blocking call in an async method. Replaced with await StreamReader.ReadToEndAsync.

diff --git a/Engine/Extras/P4VUtils/Commands/PreflightCommand.cs b/Engine/Extras/P4VUtils/Commands/PreflightCommand.cs
--- a/Engine/Extras/P4VUtils/Commands/PreflightCommand.cs
+++ b/Engine/Extras/P4VUtils/Commands/PreflightCommand.cs
@@ -134,7 +134,7 @@ namespace P4VUtils.Commands
 							"[CANCEL] - To cancel the request";
 					string Caption = "Your CL will fail to auto-submit unless fixed";
 
-					string Response = null;
+					string Response = String.Empty;
 
 					if (OperatingSystem.IsWindows())
 					{
@@ -152,8 +152,8 @@ namespace P4VUtils.Commands
 						using (Process Proc = Process.Start(SI)!)
 						{
 							StreamReader Reader = Proc.StandardOutput;
-							Response = Reader.ReadToEnd().Trim();
-							Response = Response.Replace("button returned:", "");
+							Response = (await Reader.ReadToEndAsync()).Trim();
+							Response = Response.Replace("button returned:", "", StringComparison.OrdinalIgnoreCase);
 						}
 					}
 					else
diff --git a/Engine/Extras/P4VUtils/Program.cs b/Engine/Extras/P4VUtils/Program.cs
--- a/Engine/Extras/P4VUtils/Program.cs
+++ b/Engine/Extras/P4VUtils/Program.cs
@@ -222,7 +222,7 @@ namespace P4VUtils
 
 		static string GetToolName(XmlElement ToolNode)
 		{
-			return ToolNode.SelectSingleNode("Definition").SelectSingleNode("Name").InnerText;
+			return ToolNode.SelectSingleNode("Definition")?.SelectSingleNode("Name")?.InnerText ?? string.Empty;
 		}
 
 		// returns true if all tools were removed
@@ -230,8 +230,15 @@ namespace P4VUtils
 		{
 			int ToolsChecked = 0;
 			int ToolsRemoved = 0;
+
+			XmlNodeList? CustomToolDefList = RootNode.SelectNodes("CustomToolDef");
+			if (CustomToolDefList == null)
+			{
+				return false;
+			}
+
 			// Removes tools explicitly calling the assembly location identified above - i assume as a way to "filter" only those we explicitly added (@Ben.Marsh) - nochecking, remove this comment once verified.
-			foreach (XmlNode? ChildNode in RootNode.SelectNodes("CustomToolDef"))
+			foreach (XmlNode? ChildNode in CustomToolDefList)
 			{
 				XmlElement? ChildElement = ChildNode as XmlElement;
 				if (ChildElement != null)
@@ -356,7 +363,13 @@ namespace P4VUtils
 		}
 		static void RemoveCustomToolsFromFolders(XmlElement RootNode, FileReference DotNetLocation, FileReference AssemblyLocation, ILogger Logger)
 		{
-			foreach (XmlNode? ChildNode in RootNode.SelectNodes("CustomToolFolder"))
+			XmlNodeList? CustomToolFolderList = RootNode.SelectNodes("CustomToolFolder");
+			if(CustomToolFolderList == null)
+			{
+				return;
+			}
+
+			foreach (XmlNode? ChildNode in CustomToolFolderList)
 			{
 				if (ChildNode != null)
 				{
