Rework the logic for finding the project associated with each package file to match the native code, including changes found in CL 20346906)

#rnx

- The native code changes we are trying to match from CL 20346906, fix a problem when parsing a package path containing more than one content directory.
- The method now takes an ILogger so that we can record problems found.
- Even though we don't use the .uplugin file, we now match the native code and look for it first, rather than looking for the /plugins/ directory immediately. This should reduce any differences between the two code paths.

diff --git a/Engine/Extras/P4VUtils/Commands/SubmitCommand.cs b/Engine/Extras/P4VUtils/Commands/SubmitCommand.cs
--- a/Engine/Extras/P4VUtils/Commands/SubmitCommand.cs
+++ b/Engine/Extras/P4VUtils/Commands/SubmitCommand.cs
@@ -289,7 +289,7 @@ namespace P4VUtils.Commands
 			foreach (string path in packagePaths)
 			{
 				string normalizedPath = NormalizeFilename(path);
-				string projectFilePath = FindProjectForPackage(normalizedPath);
+				string projectFilePath = FindProjectForPackage(normalizedPath, logger);
 
 				if (!String.IsNullOrEmpty(projectFilePath))
 				{
@@ -454,53 +454,70 @@ namespace P4VUtils.Commands
 		/// Finds the unreal project file for a given unreal package
 		/// </summary>
 		/// <param name="packagePath">The package to find the project for</param>
+		/// <param name="logger">Interface for logging</param>
 		/// <returns>The path of the projects .uproject file if found, an empty string if no valid project file was found</returns>
-		private static string FindProjectForPackage(string packagePath)
+		private static string FindProjectForPackage(string packagePath, ILogger logger)
 		{
 			// @todo note that this function mirrors FUnrealVirtualizationToolApp::TryFindProject
 			// both will not work correctly with some plugin setups, this is known and will be
 			// fixed later when FUnrealVirtualizationToolApp is also fixed.
 			packagePath = NormalizeFilename(packagePath);
 
-			int contentIndex = packagePath.IndexOf("/content/", StringComparison.OrdinalIgnoreCase);
+			int contentIndex = packagePath.LastIndexOf("/content/", StringComparison.OrdinalIgnoreCase);
 			if (contentIndex == -1)
 			{
-				// Error condition no content directory?
+				logger.LogWarning("'{Path}' is not under a content directory", packagePath);
 				return string.Empty;
 			}
 
-			string projectDirectory = packagePath[..contentIndex];
+			while (contentIndex != -1)
+			{
+				// Assume that the project directory is the parent of the /content/ directory
+				string projectDirectory = packagePath[..contentIndex];
 
-			string[] projectFiles = Directory.GetFiles(projectDirectory, "*.uproject");
+				string[] projectFiles = Directory.GetFiles(projectDirectory, "*.uproject");
 
-			if (projectFiles.Length == 0)
-			{
-				// If there was no project file, the package could be in a plugin, so lets check for that
-				int pluginIndex = packagePath.IndexOf("/plugins/", StringComparison.OrdinalIgnoreCase);
-				if (pluginIndex == -1)
+				if (projectFiles.Length == 0)
+				{
+					// If there was no project file, the package could be in a plugin, so lets check for that
+					string pluginDirectory = projectDirectory;
+					string[] pluginFiles = Directory.GetFiles(pluginDirectory, "*.uplugin");
+
+					if (pluginFiles.Length == 1)
+					{
+						// We have a valid plugin file, so we should be able to find a /plugins/ directory which will be just below the project directory
+						int pluginIndex = pluginDirectory.LastIndexOf("/plugins/", StringComparison.OrdinalIgnoreCase);
+						if (pluginIndex != -1)
+						{
+							// We found the plugin root directory so the one above it should be the project directory
+							projectDirectory = pluginDirectory[..pluginIndex];
+							projectFiles = Directory.GetFiles(projectDirectory, "*.uproject");
+						}
+					}
+					else if (pluginFiles.Length > 1)
+					{
+						logger.LogWarning("Found multiple .uplugin files for '{Path}' at '{PluginDir}'", pluginDirectory);
+						return string.Empty;
+					}
+				}
+
+				if (projectFiles.Length == 1)
 				{
-					// Error condition not a plugin?
+					return NormalizeFilename(projectFiles[0]);
+				}
+				else if (projectFiles.Length > 1)
+				{
+					logger.LogWarning("Found multiple .uproject files for '{Path}' at '{ProjectDir}'", projectDirectory);
 					return string.Empty;
 				}
 
-				projectDirectory = packagePath[..pluginIndex];
-				projectFiles = Directory.GetFiles(projectDirectory, "*.uproject");
+				// Could be more than one content directory in the path so lets keep looking
+				contentIndex = packagePath.LastIndexOf("/content/", contentIndex, StringComparison.OrdinalIgnoreCase);
 			}
 
-			if (projectFiles.Length == 1)
-			{
-				return NormalizeFilename(projectFiles[0]);
-			}
-			else if (projectFiles.Length > 1)
-			{
-				// Error condition too many project files?
-				return string.Empty;
-			}
-			else
-			{
-				// Error condition no project file found
-				return string.Empty;
-			}
+			// We found one or more content directories but none of them contained a project file
+			logger.LogWarning("FFailed to find project file for '{Package}'", packagePath);
+			return string.Empty;
 		}
 
 		private static string GetEngineRootForProject(string projectFilePath, ILogger logger)
