Allow the P4VUtils 'Virtualize and Submit' command to work on multiple changelists at the same time.

#rnx

- If we use uppercase %C in P4V to pass the changelist number then we will get a list of all changelists that the user has selected rather than launching the command once per changelist which is what happens when the lower case %c is used
-- Due to the use of UBT, we can run into problems if the user tries to submit multiple changelists and they are processed in different processes. Even if we fixed the command to only run UBT one process at a time, running many versions of the submit tool will create many more perforce connections which we'd like to avoid, so having all of the changelists run through one instance is still prefered.
- Processing a changelist has been moved to a new method, ::ProcessChangelist which is pretty much the same code as we had in  ::Execute.
-- Rather than connecting to perforce for each changelist, this code has been kept in ::Execute and is only run once, the connection is then passed to ::ProcessChangelist.
- Logging has been improved to show when the tool is exiting, either from failure or success. The p4v terminal does not do a good job informating the user when the process has finished running so hopefully this will help.
- A small extension to ILogger has been added, ::NewLine which writes an empty line to the logger.
- This was the simplest way to add support to multiple changelists. We process each one by itself until we run out of changelists or we reach an error. However it is likely that packages in multiple changelists would be in the same project or at least use the same engine installation, so we could remove some overhead by parsing all changelists, gathering the packages, virtualizating those packages then submitting the changelists afterwards.
-- This would be a great approach for the average user as they will likely be submitting relatively small changelists (< 100 files) but would be a terrible approach if the entire project was being virtualized where the total number of packages being submitting could be in the 100k-million range. So we'd probably want a "process changelists separatly" mode even if we did make the optimized version.
-- Looking into adding a combined changelist mode will be added as an optimization task.
- Added a new preprocessor define 'UE_DEBUG_DISABLE_SUBMITS' that when enabled will cause the command to only pretend to submit the target changelist(s) this is intended for use in development to avoid submitting changelists over and over when testing the non-submission portion of the command.
-- Note that it is defined inside of a DEBUG check, so even if it is accidently submitted enabled, it will not be published.

diff --git a/Engine/Extras/P4VUtils/Commands/SubmitCommand.cs b/Engine/Extras/P4VUtils/Commands/SubmitCommand.cs
--- a/Engine/Extras/P4VUtils/Commands/SubmitCommand.cs
+++ b/Engine/Extras/P4VUtils/Commands/SubmitCommand.cs
@@ -1,5 +1,11 @@
 // Copyright Epic Games, Inc. All Rights Reserved.
 
+#if DEBUG
+// When enabled the command will not actually submit the changelist, but will act like it did.
+// This is for development purposes only and so is only available in the debug config.
+//#define UE_DEBUG_DISABLE_SUBMITS
+#endif
+
 using EpicGames.Core;
 using EpicGames.Perforce;
 using Microsoft.Extensions.Logging;
@@ -13,51 +19,105 @@ using System.Threading;
 using System.Threading.Tasks;
 
 namespace P4VUtils.Commands
-{
+{	
+	public static class LoggingUtils
+	{
+		/// <summary>
+		/// This extension is used to write a new line to the logger. 
+		/// log.NewLine better shows intent over log.LogInformation("")
+		/// </summary>
+		public static void NewLine(this ILogger logger)
+		{
+			logger.LogInformation("");
+		}
+	}
+
 	class SubmitAndVirtualizeCommand : Command
 	{
 		public override string Description => "Virtualize And Submit";
 
-		public override CustomToolInfo CustomTool => new CustomToolInfo("Virtualize And Submit", "$c %c") { ShowConsole = true, RefreshUI = true };
+		public override CustomToolInfo CustomTool => new CustomToolInfo("Virtualize And Submit", "$c %C") { ShowConsole = true, RefreshUI = true };
 
 		public override async Task<int> Execute(string[] args, IReadOnlyDictionary<string, string> configValues, ILogger logger)
 		{
 			// Parse command lines
 			if (args.Length < 3)
 			{
-				logger.LogError("Not enough args for command!");
+				logger.LogError("Not enough args for command, tool is now exiting");
 				return 1;
 			}
 
 			string clientSpec = args[1];
 
-			if (!int.TryParse(args[2], out int changeNumber))
+			List<int> changelistsToSubmit = new List<int>();
+
+			for (int index = 2; index < args.Length; ++index)
 			{
-				logger.LogError("'{Argument}' is not a numbered changelist", args[1]);
-				return 1;
+				if (!int.TryParse(args[index], out int changeNumber))
+				{
+					logger.LogError("'{Argument}' is not a numbered changelist, tool is now exiting", args[index]);
+					return 1;
+				}
+
+				changelistsToSubmit.Add(changeNumber);
 			}
 
-			logger.LogInformation("Attempting to virtualize and submit changelist {Change} in the workspace {Spec}", changeNumber, clientSpec);
+			if (changelistsToSubmit.Count == 0)
+			{
+				logger.LogError("No changelists to submit were provided, tool is now exiting");
+				return 1;
+			}
 
 			// Connect to perforce and validate
 
-			logger.LogInformation("Connecting to perforce...");
-			
+			logger.LogInformation("Connecting to Perforce...");
+
 			// We prefer the native client to avoid the problem where different versions of p4.exe expect
 			// or return records with different formatting to each other.
 			PerforceSettings settings = new PerforceSettings(PerforceSettings.Default) { PreferNativeClient = true, ClientName = clientSpec };
 			using IPerforceConnection perforceConnection = await PerforceConnection.CreateAsync(settings, logger);
 			if (perforceConnection == null)
 			{
-				logger.LogError("Failed to connect", args[1]);
+				logger.LogError("Failed to connect to Perforce, tool is now exiting");
 				return 1;
 			}
 
+			for (int index = 0; index < changelistsToSubmit.Count; ++index)
+			{
+				logger.NewLine();
+				logger.LogInformation("========== Processing change list {Index}/{Count} ==========", index + 1, changelistsToSubmit.Count);
+
+				int changeNumber = changelistsToSubmit[index];
+				if(await ProcessChangelist(perforceConnection, changeNumber, logger) == false)
+				{
+					logger.LogError("Failed to process {Changelist}, tool is now exiting", changeNumber);
+					return 1;
+				}
+			}
+
+			logger.NewLine();
+			logger.LogInformation("========== All changelist(s) submitted successfully, tool is now exiting ==========");
+			
+			return 0;
+		}
+
+		public static async Task<bool> ProcessChangelist(IPerforceConnection perforceConnection, int changeNumber, ILogger logger)
+		{
+			if (perforceConnection == null || perforceConnection.Settings == null || perforceConnection.Settings.ClientName == null)
+			{
+				logger.LogError("Invalid Perforce connection!");
+				return false;
+			}
+
+			string clientSpec = perforceConnection.Settings.ClientName;
+
+			logger.LogInformation("Attempting to virtualize and submit changelist {Change} in the workspace {Spec}", changeNumber, clientSpec);
+
 			// First we need to find the packages in the changelist
 			string[]? localFilePaths = await FindPackagesInChangelist(perforceConnection, changeNumber, logger);
 			if (localFilePaths == null)
 			{
-				return 1;
+				return false;
 			}
 
 			if (localFilePaths.Length > 0)
@@ -69,7 +129,7 @@ namespace P4VUtils.Commands
 				if (await DoesChangelistHaveShelvedFiles(perforceConnection, changeNumber) == true)
 				{
 					logger.LogError("Changelist {Change} has shelved files and cannot be submitted", changeNumber);
-					return 1;
+					return false;
 				}
 
 				logger.LogInformation("Found {Amount} package(s) that may need virtualization", localFilePaths.Length);
@@ -89,7 +149,8 @@ namespace P4VUtils.Commands
 					string engineRoot = GetEngineRootForProject(project.Key, logger);
 					if (!String.IsNullOrEmpty(engineRoot))
 					{
-						logger.LogInformation("\nAttempting to virtualize packages in project '{Project}' via the engine installation '{Engine}'", project.Key, engineRoot);
+						logger.NewLine();
+						logger.LogInformation("Attempting to virtualize packages in project '{Project}' via the engine installation '{Engine}'", project.Key, engineRoot);
 						// @todo Many projects can share the same engine install, and technically UnrealVirtualizationTool
 						// supports the virtualization files from many projects at the same time. We could consider doing
 						// this pass per engine install rather than per project? At the very least we should only 'build'
@@ -102,7 +163,7 @@ namespace P4VUtils.Commands
 						// Check if the compilation of the tool succeeded or not
 						if (await compileResult == false)
 						{
-							return 1;
+							return false;
 						}
 
 						// Even though this will have been done while we were waiting for BuildVirtualizationTool to complete 
@@ -113,7 +174,7 @@ namespace P4VUtils.Commands
 
 						if (await RunVirtualizationTool(engineRoot, clientSpec, tempFilesPath, logger) == false)
 						{
-							return 1;
+							return false;
 						};
 
 						// @todo Make sure this always gets cleaned up when we go out of scope
@@ -122,22 +183,23 @@ namespace P4VUtils.Commands
 					else
 					{
 						logger.LogError("Failed to find engine root for project {Project}", project.Key);
-						return 1;
+						return false;
 					}
 				}
 
-				logger.LogInformation("\nAll packages have been virtualized");
+				logger.NewLine();
+				logger.LogInformation("All packages have been virtualized");
 
 				//@todo ideally we should get the tags back from UnrealVirtualizationTool
 				ChangeRecord? changeRecord = await StampChangelistDescription(perforceConnection, changeNumber, logger);
 				if (changeRecord == null)
 				{
-					return 1;
+					return false;
 				}
 
 				logger.LogInformation("Attempting to submit changelist {Number}...", changeNumber);
 
-				// Submit
+#if !UE_DEBUG_DISABLE_SUBMITS
 				if (await SubmitChangelist(perforceConnection, changeNumber, logger) == false)
 				{
 					// If the final submit failed we remove the virtualization tag, even though the changelist is technically
@@ -151,8 +213,12 @@ namespace P4VUtils.Commands
 					}
 					logger.LogInformation("Virtualization tags have been removed.");
 
-					return 1;
+					return false;
 				}
+#else
+				// Dummy pretend submit for debugging purposes, allows developer iteration without actually submitting
+				logger.LogInformation("Successfully submited changelist {Change}", changeNumber);
+#endif // !UE_DEBUG_DISABLE_SUBMITS
 			}
 			else
 			{
@@ -160,11 +226,11 @@ namespace P4VUtils.Commands
 
 				if (await SubmitChangelist(perforceConnection, changeNumber, logger) == false)
 				{
-					return 1;
+					return false;
 				}
 			}
 		
-			return 0;
+			return true;
 		}
 
 		/// <summary>
