Allow the P4VUtils 'Virtualize and Submit' command to use precompiled binaries is they are available.

#rnx

- Instead of just building the UnrealVirtualizationTool we will now check to see if there is a precompiled version we can use.
-- If the user has rebuilt the editor locally, then we will ignore any precompiled version of UnrealVirtualizationTool and rebuild that aswell incase there are code changes that it needs to pick up.
- To tell if an exe is locally built or a precompiled binary we check for a .version OR .target file (this can vary depending how a program is set up) then look for the JSON property 'IsPromotedBuild ' which when none zero means that the binary is probably prebuilt.
-- This isn't very robust and I am not happy with it, but we don't seem to have an official way to determine a precompiled binary. When we do get an offical way to do it I will update the command.

diff --git a/Engine/Extras/P4VUtils/Commands/SubmitCommand.cs b/Engine/Extras/P4VUtils/Commands/SubmitCommand.cs
--- a/Engine/Extras/P4VUtils/Commands/SubmitCommand.cs
+++ b/Engine/Extras/P4VUtils/Commands/SubmitCommand.cs
@@ -15,11 +15,37 @@ using System.Collections.Generic;
 using System.IO;
 using System.Linq;
 using System.Text;
+using System.Text.Json;
 using System.Threading;
 using System.Threading.Tasks;
 
 namespace P4VUtils.Commands
-{	
+{
+	/// <summary>
+	/// Contains info about a built exe, generated by our build system (found in .version files)
+	/// </summary>
+	public class BuildVersion
+	{
+		public int? MajorVersion { get; set; }
+		public int? MinorVersion { get; set; }
+		public int? PatchVersion { get; set; }
+		public int? Changelist { get; set; }
+		public int? CompatibleChangelist { get; set; }
+		public int? IsLicenseeVersion { get; set; }
+		public int? IsPromotedBuild { get; set; }
+		public string? BranchName { get; set; }
+		public string? BuildId { get; set; }
+	}
+
+	/// <summary>
+	/// A partial representation of the .target file generated by our build system. We are only
+	/// interested in the BuildVersion info and so the class does not have any other members.
+	/// </summary>
+	public class TargetBuildVersion
+	{
+		public BuildVersion? Version { get; set; }
+	}
+
 	public static class LoggingUtils
 	{
 		/// <summary>
@@ -156,7 +182,7 @@ namespace P4VUtils.Commands
 						// this pass per engine install rather than per project? At the very least we should only 'build'
 						// the tool once per engine
 
-						Task<bool> compileResult = BuildVirtualizationTool(engineRoot, logger);
+						Task<bool> compileResult = ValidateOrBuildVirtualizationTool(engineRoot, logger);
 
 						string tempFilesPath = await WritePackageFileList(project.Value, logger);
 
@@ -233,6 +259,122 @@ namespace P4VUtils.Commands
 			return true;
 		}
 
+		/// <summary>
+		/// Checks to see if a precompiled version of the UnrealVirtualizationTool exists and if so
+		/// if it is valid for us to use this.
+		/// If we cannot use a precompiled version of the tool then we will request that it be 
+		/// compiled locally.
+		/// </summary>
+		/// <param name="engineRoot">Root path of the engine we want to build the tool for</param>
+		/// <param name="logger">Interface for logging</param>
+		/// <returns>True if the UnrealVirtualizationTool.exe exists in a form that we can use, otherwise false</returns>
+		private static async Task<bool> ValidateOrBuildVirtualizationTool(string engineRoot, ILogger logger)
+		{
+			// Check to see if the precompiled binaries for the  editor and virtualization tool are available
+			bool hasEditorPCB = IsPrecompiledVersionAvailable(engineRoot, "UnrealEditor", logger);
+			bool hasToolPCB = IsPrecompiledVersionAvailable(engineRoot, "UnrealVirtualizationTool", logger);
+
+			// If the editor is not a precompiled binary, then we know that the user has compiled it locally.
+			// This means that even if we have a precompiled version of UnrealVirtualizationTool we cannot
+			// use it as it's virtualization process might be 'out of date' compared to the locally compiled
+			// editor.
+			// This means we can only use the precompiled binary for UnrealVirtualizationTool if UnrealEditor
+			// is also a precompiled binary.
+			if (hasEditorPCB && hasToolPCB)
+			{
+				logger.LogInformation("Using a precompiled binary version of UnrealVirtualizationTool");
+				return true;
+			}
+			else if (hasToolPCB)
+			{
+				// Let the user know why we are compiling UnrealVirtualizationTool even though there is a precompiled
+				// version available.
+				logger.LogInformation("Cannot use the precompiled version of UnrealVirtualizationTool as the UnrealEditor has been locally compiled");
+			}
+
+			// We must try to build the tool locally, if the user does not have a valid code compilation 
+			// tool chain installed then UnrealBuildTool will give them errors.
+			return await BuildVirtualizationTool(engineRoot, logger);
+		}
+
+		/// <summary>
+		/// Checks if there is a precompiled version of a given program available for use.
+		/// </summary>
+		/// <param name="engineRoot">Root path of the engine we want to build the tool for</param>
+		/// <param name="programName">Name of the program to check for</param>
+		/// <param name="logger">Interface for logging</param>
+		/// <returns>True if there is a version of the program that was precompiled, otherwise false</returns>
+		private static bool IsPrecompiledVersionAvailable(string engineRoot, string programName, ILogger logger)
+		{
+			// Depending on how a program is set up it might have a .version file or a .target file that contains info
+			// about the exe itself. From this we can find a JSON property "IsPromotedBuild" which when none zero 
+			// means that the build was precompiled.
+			// This isn't the most robust check but should work for now until we can add an 'offical' codified way to
+			// check for precompiled binaries.
+
+			// First we check for a .version file for the given program
+			string versionPath = String.Format(@"{0}\Engine\Binaries\Win64\{1}.version", engineRoot, programName);
+			
+			if (System.IO.File.Exists(versionPath))
+			{ 
+				using (StreamReader inputStream = new StreamReader(versionPath))
+				{
+					string jsonObject = inputStream.ReadToEnd();
+
+					BuildVersion? buildVersion = null;
+					try
+					{
+						buildVersion = JsonSerializer.Deserialize<BuildVersion>(jsonObject);
+					}
+					catch (System.Text.Json.JsonException)
+					{
+						// No need to print the exception as the user will not be expected to fix the file, a warning will be shown below
+					}
+
+					if (buildVersion == null || buildVersion.IsPromotedBuild == null)
+					{
+						logger.LogWarning("Unable to read build version from '{FilePath}'", versionPath);
+					}
+					else if (buildVersion.IsPromotedBuild != 0)
+					{
+						return true;
+					}
+				}
+			}
+
+			// We failed to find what we needed from the .version file, so try the .target file instead
+			string targetPath = String.Format(@"{0}\Engine\Binaries\Win64\{1}.target", engineRoot, programName);
+
+			if (System.IO.File.Exists(targetPath))
+			{
+				using (StreamReader inputStream = new StreamReader(targetPath))
+				{
+					string jsonObject = inputStream.ReadToEnd();
+
+					TargetBuildVersion? targetBuildVersion = null;
+					try
+					{
+						targetBuildVersion = JsonSerializer.Deserialize<TargetBuildVersion>(jsonObject);
+					}
+					catch (System.Text.Json.JsonException)
+					{
+						// No need to print the exception as the user will not be expected to fix the file, a warning will be shown below
+					}
+					
+					if (targetBuildVersion == null || targetBuildVersion.Version == null || targetBuildVersion.Version.IsPromotedBuild == null)
+					{
+						logger.LogWarning("Unable to read build version from '{FilePath}'", targetPath);
+					}
+					else if (targetBuildVersion.Version.IsPromotedBuild != 0)
+					{
+						return true;
+					}
+				}
+			}
+
+			return false;
+		}
+
 		/// <summary>
 		/// Compiles UnrealVirtualizationTool via RunUBT.bat
 		/// </summary>
