Fix P4VUtils failing to remove its old commands when running an install

#rnx

- When UGS invokes the tool to be installed, it will call 'P4VUtils.exe install' which meant that getting the process path as the dotnet path would return the exe
- However previous versions of the tool were creating a path to dotnet.exe to use as the dotnet path
- When removing the old commands before installing the new ones, one of the filters is to check the dotnet path of the command and make sure that it matches the current one
- So if the older version was using a path to dotnet.exe and the newer version, the p4vutils.exe, they won't match and the old commands won't be removed, so when the new ones are added we end up with duplicates.
- For now the easiest fix is to go back to the logic before change 20248689, so on windows we make sure that we find dotnet.exe so it will match.

diff --git a/Engine/Extras/P4VUtils/Program.cs b/Engine/Extras/P4VUtils/Program.cs
--- a/Engine/Extras/P4VUtils/Program.cs
+++ b/Engine/Extras/P4VUtils/Program.cs
@@ -429,8 +429,17 @@ namespace P4VUtils
 			}
 
 
-			// so we don't need to run SetupDotnet.sh every time we run a tool, point to the executing dotnet, assuming it will be usable later
-			FileReference DotNetLocation = new FileReference(Environment.ProcessPath!);
+			FileReference DotNetLocation;
+			if (OperatingSystem.IsWindows())
+			{
+				DotNetLocation = FileReference.Combine(DirectoryReference.GetSpecialFolder(Environment.SpecialFolder.ProgramFiles)!, "dotnet", "dotnet.exe");
+			}
+			else
+			{
+				// so we don't need to run SetupDotnet.sh every time we run a tool, point to the executing dotnet, assuming it will be usable later
+				DotNetLocation = new FileReference(Environment.ProcessPath!);
+			}
+
 			FileReference AssemblyLocation = new FileReference(Assembly.GetExecutingAssembly().GetOriginalLocation());
 
 			XmlElement? Root = Document.SelectSingleNode("CustomToolDefList") as XmlElement;
