P4VUtils will no longer fail to remove commands that contain quoted string

#rnx

- This is a follow on fix from 20451429, although that solved the problem of duplicate registered commands due to differing DotNet locations it did not solve a problem introduced in 20248689 caused by quoting the command argument.
- If the path in the command element does get quoted due to containing a space in the file path then FileReference will fail to parse it as RemoveCustomToolsFromNode intended and the comparison between the Command element and the DotNet location would always fail.
-- This meant that a newly published version of the tool would replace the old commands, but when published a second time would start to register duplicate commands.
- The simple solution is to try and strip the quotes (if any) before creating the FileReference for comparison.

diff --git a/Engine/Extras/P4VUtils/Program.cs b/Engine/Extras/P4VUtils/Program.cs
--- a/Engine/Extras/P4VUtils/Program.cs
+++ b/Engine/Extras/P4VUtils/Program.cs
@@ -245,7 +245,15 @@ namespace P4VUtils
 				{
 					ToolsChecked++;
 					XmlElement? CommandElement = ChildElement.SelectSingleNode("Definition/Command") as XmlElement;
-					if (CommandElement != null && new FileReference(CommandElement.InnerText) == DotNetLocation)
+
+
+					// In a recent change we started to output the Command element as a quoted argument if the path contains spaces.
+					// FileReference does not resolve quoted string properly which was causing the comparisons here to fail.
+					// We can strip the quotes before creating a FileReference to compare with DotNetLocation to ensure that the comparsion
+					// is correct.
+					String CommandPath = (CommandElement?.InnerText ?? String.Empty).StripQuoteArgument();
+
+					if (new FileReference(CommandPath) == DotNetLocation)
 					{
 						XmlElement? ArgumentsElement = ChildElement.SelectSingleNode("Definition/Arguments") as XmlElement;
 						if (ArgumentsElement != null)
