- Reverted the P4VUtils changes that needed dotnet6, since some users don't have it installed, and we can't use the dotnet checked in to p4, since P4VUtils are not tied to any stream

diff --git a/Engine/Extras/P4VUtils/Commands/CopyCLCommand.cs b/Engine/Extras/P4VUtils/Commands/CopyCLCommand.cs
--- a/Engine/Extras/P4VUtils/Commands/CopyCLCommand.cs
+++ b/Engine/Extras/P4VUtils/Commands/CopyCLCommand.cs
@@ -37,13 +37,12 @@ namespace P4VUtils.Commands
 				return 1;
 			}
 
-#if IS_WINDOWS
 			// Jump through some hoops to make async & windows com happy and force STA on the clipboard call
 			var tcs = new TaskCompletionSource<int>();
 
 			var SetClipThread = new Thread(() => 
 				{ 
-					System.Windows.Forms.Clipboard.SetText(Change.ToString());
+					Clipboard.SetText(Change.ToString());
 					tcs.SetResult(0);
 				}
 			);
@@ -54,10 +53,6 @@ namespace P4VUtils.Commands
 
 			await tcs.Task;
 
-#elif IS_MAC
-			System.Diagnostics.Process.Start("/bin/bash", $" -c \"osascript -e 'set the clipboard to \\\"{Change}\\\"'\"");
-#endif
-
 			return 0;
 		}
 	}
diff --git a/Engine/Extras/P4VUtils/Commands/PreflightCommand.cs b/Engine/Extras/P4VUtils/Commands/PreflightCommand.cs
--- a/Engine/Extras/P4VUtils/Commands/PreflightCommand.cs
+++ b/Engine/Extras/P4VUtils/Commands/PreflightCommand.cs
@@ -5,34 +5,22 @@ using EpicGames.Perforce;
 using Microsoft.Extensions.Logging;
 using System;
 using System.Collections.Generic;
+using System.ComponentModel;
 using System.Diagnostics;
+using System.IO;
+using System.Reflection;
 using System.Runtime.InteropServices;
+using System.Text;
+using System.Text.RegularExpressions;
 using System.Threading;
 using System.Threading.Tasks;
+using System.Windows;
 using System.Linq;
-using System.IO;
-
-#if IS_WINDOWS
-using System.Windows.Forms;
-#endif
 
 namespace P4VUtils.Commands
 {
 	class PreflightCommand : Command
 	{
-#if IS_WINDOWS
-		[System.Runtime.InteropServices.DllImport("user32.dll")]
-		private static extern bool SetProcessDPIAware();
-
-		public static void SetupVisuals()
-		{
-			// make the form look good on modern displays!
-			SetProcessDPIAware();
-
-			Application.EnableVisualStyles();
-			Application.SetCompatibleTextRenderingDefault(false);
-		}
-#endif
 
 		static public string StripReviewFyiHashTags(string InString)
 		{
@@ -125,57 +113,30 @@ namespace P4VUtils.Commands
 
 				if (OpenedRecords.Count > 0 && IsSubmit())
 				{
-					string Prompt = "Your CL was just shelved however it still has files checked out\r\n" +
-							"If the files remain in the CL your preflight will fail to submit\r\n" +
-							"\r\n" +
-							"Click:\r\n" +
-							"[YES] - To revert local files and submit the preflight,\r\n" +
-							"[NO] - To start the preflight, and move the files manually\r\n" +
-							"[CANCEL] - To cancel the request";
-					string Caption = "Your CL will fail to auto-submit unless fixed";
-
-					string Response = String.Empty;
-
-					if (OperatingSystem.IsWindows())
+					MessageBoxResult result= MessageBox.Show(
+						"Your CL was just shelved however it still has files checked out\r\n" +
+						"If the files remain in the CL your preflight will fail to submit\r\n" +
+						"\r\n" +
+						"Click:\r\n" +
+						"[YES] - To revert local files and submit the preflight,\r\n" +
+						"[NO] - To start the preflight, and move the files manually\r\n" +
+						"[CANCEL] - To cancel the request", 
+
+						"Your CL will fail to auto-submit unless fixed", 
+						MessageBoxButton.YesNoCancel);
+
+					if (result == MessageBoxResult.Cancel)
 					{
-#if IS_WINDOWS
-						SetupVisuals();
-
-						DialogResult result= MessageBox.Show(Prompt, Caption, MessageBoxButtons.YesNoCancel);
-						Response = result.ToString();
-#endif
-					}
-					else if (OperatingSystem.IsMacOS())
-					{
-						ProcessStartInfo SI = new ProcessStartInfo("/bin/bash", $"-c \" osascript -e 'display dialog \\\"{Prompt}\\\" buttons {{\\\"Yes\\\", \\\"No\\\", \\\"Cancel\\\"}} with icon caution'\"");
-						SI.RedirectStandardOutput = true;
-						using (Process Proc = Process.Start(SI)!)
-						{
-							StreamReader Reader = Proc.StandardOutput;
-							Response = (await Reader.ReadToEndAsync()).Trim();
-							Response = Response.Replace("button returned:", "", StringComparison.OrdinalIgnoreCase);
-						}
-					}
-					else
-					{
-						Logger.LogWarning("Your CL was just shelved however it still has files checked out\r\n" +
-							"If the files remain in the CL your preflight will fail to submit\r\n");
-					}
-
-
-					if (Response == "No")
-					{
-						// do nothing - user has been warned.
+						Logger.LogInformation("User Opted to cancel");
+						return 0;
 					}
-					else if (Response == "Yes")
+					else if(result == MessageBoxResult.Yes)
 					{
 						await Perforce.RevertAsync(Change, null, RevertOptions.None, OpenedRecords.Select(x => x.ClientFile!).ToArray(), CancellationToken.None);
 					}
-					// any other reply is Cancel (like on Mac, hitting Escape will return a weird string, not Cancel)
 					else
 					{
-						Logger.LogInformation("User Opted to cancel");
-						return 0;
+						// do nothing - user has been warned.
 					}
 				}
 			}
diff --git a/Engine/Extras/P4VUtils/P4VUtils.csproj b/Engine/Extras/P4VUtils/P4VUtils.csproj
--- a/Engine/Extras/P4VUtils/P4VUtils.csproj
+++ b/Engine/Extras/P4VUtils/P4VUtils.csproj
@@ -1,41 +1,10 @@
-﻿<Project Sdk="Microsoft.NET.Sdk">
-
-    <PropertyGroup>
-      <IsWindows Condition="'$([System.Runtime.InteropServices.RuntimeInformation]::IsOSPlatform($([System.Runtime.InteropServices.OSPlatform]::Windows)))' == 'true'">true</IsWindows> 
-      <IsOSX Condition="'$([System.Runtime.InteropServices.RuntimeInformation]::IsOSPlatform($([System.Runtime.InteropServices.OSPlatform]::OSX)))' == 'true'">true</IsOSX> 
-      <IsLinux Condition="'$([System.Runtime.InteropServices.RuntimeInformation]::IsOSPlatform($([System.Runtime.InteropServices.OSPlatform]::Linux)))' == 'true'">true</IsLinux> 
-    </PropertyGroup>
-
-    <Choose>
-    <When Condition="'$(IsWindows)'=='true'">
-      <PropertyGroup>
-        <TargetFramework>net6.0-windows</TargetFramework>
-        <UseWindowsForms>true</UseWindowsForms>
-        <DefineConstants>$(DefineConstants);IS_WINDOWS</DefineConstants>
-      </PropertyGroup>
-    </When>
-    <Otherwise>
-      <PropertyGroup>
-        <TargetFramework>net6.0</TargetFramework>
-      </PropertyGroup>
-      <Choose>
-      <When Condition="'$(IsOSX)'=='true'">
-        <PropertyGroup>
-          <DefineConstants>$(DefineConstants);IS_MAC</DefineConstants>
-        </PropertyGroup>
-      </When>
-      <Otherwise>
-        <PropertyGroup>
-          <DefineConstants>$(DefineConstants);IS_LINUX</DefineConstants>
-        </PropertyGroup>
-      </Otherwise>
-      </Choose>
-    </Otherwise>
-  </Choose>
+<Project Sdk="Microsoft.NET.Sdk.WindowsDesktop">
 
   <PropertyGroup>
     <OutputType>Exe</OutputType>
+    <TargetFramework>netcoreapp3.1</TargetFramework>
     <Nullable>enable</Nullable>
+    <UseWPF>true</UseWPF>
     <EnableNETAnalyzers>True</EnableNETAnalyzers>
   </PropertyGroup>
 
diff --git a/Engine/Extras/P4VUtils/Program.cs b/Engine/Extras/P4VUtils/Program.cs
--- a/Engine/Extras/P4VUtils/Program.cs
+++ b/Engine/Extras/P4VUtils/Program.cs
@@ -48,9 +48,7 @@ namespace P4VUtils
 		public static IReadOnlyDictionary<string, Command> RootHelperCommands { get; } = new Dictionary<string, Command>(StringComparer.OrdinalIgnoreCase)
 		{
 			["describe"] = new DescribeCommand(),
-#if !IS_LINUX
 			["copyclnum"] = new CopyCLCommand(),
-#endif
 		};
 
 		// UESubmit - commands that help with submitting files/changelists
@@ -168,7 +166,7 @@ namespace P4VUtils
 				return await Command.Execute(Args, ConfigValues, Logger);
 			}
 			else
-			{
+		{
 				Logger.LogError("Unknown command: {Command}", Args[0]);
 				PrintHelp(Logger);
 				return 1;
@@ -412,15 +410,7 @@ namespace P4VUtils
 				return 1;
 			}
 
-			FileReference ConfigFile;
-			if (OperatingSystem.IsMacOS())
-			{
-				ConfigFile = FileReference.Combine(ConfigDir, "Library", "Preferences", "com.perforce.p4v", "customtools.xml");
-			}
-			else
-			{
-				ConfigFile = FileReference.Combine(ConfigDir, ".p4qt", "customtools.xml");
-			}
+			FileReference ConfigFile = FileReference.Combine(ConfigDir, ".p4qt", "customtools.xml");
 
 			XmlDocument Document = new XmlDocument();
 			if (!TryLoadXmlDocument(ConfigFile, Document))
@@ -437,17 +427,7 @@ namespace P4VUtils
 			}
 
 
-			FileReference DotNetLocation;
-			if (OperatingSystem.IsWindows())
-			{
-				DotNetLocation = FileReference.Combine(DirectoryReference.GetSpecialFolder(Environment.SpecialFolder.ProgramFiles)!, "dotnet", "dotnet.exe");
-			}
-			else
-			{
-				// so we don't need to run SetupDotnet.sh every time we run a tool, point to the executing dotnet, assuming it will be usable later
-				DotNetLocation = new FileReference(Environment.ProcessPath!);
-			}
-
+			FileReference DotNetLocation = FileReference.Combine(DirectoryReference.GetSpecialFolder(Environment.SpecialFolder.ProgramFiles)!, "dotnet", "dotnet.exe");
 			FileReference AssemblyLocation = new FileReference(Assembly.GetExecutingAssembly().GetOriginalLocation());
 
 			XmlElement? Root = Document.SelectSingleNode("CustomToolDefList") as XmlElement;
