Added safety checks to backout command:
- check if the CL desc indicates that it was a robomerge or a cherry pick (ushell or p4v)
- check that the backout is of an edit operation (edit/add/delete) and not an integration
- check that the CL doesn't contain objeect or engine version files
- throw a dialogue if any of the above checks considers the backout 'unsafe'

Fixed revert exception when the backout fails due to files being checked out (also added warning similar to P4V's undo operation)

#rnx

diff --git a/Engine/Extras/P4VUtils/Commands/BackoutCommand.cs b/Engine/Extras/P4VUtils/Commands/BackoutCommand.cs
--- a/Engine/Extras/P4VUtils/Commands/BackoutCommand.cs
+++ b/Engine/Extras/P4VUtils/Commands/BackoutCommand.cs
@@ -3,6 +3,7 @@
 using EpicGames.Core;
 using EpicGames.Perforce;
 using Microsoft.Extensions.Logging;
+using P4VUtils.Perforce;
 using System;
 using System.Collections.Generic;
 using System.ComponentModel;
@@ -10,15 +11,86 @@ using System.Linq;
 using System.Text;
 using System.Threading;
 using System.Threading.Tasks;
+using System.Windows;
 
 namespace P4VUtils.Commands
 {
+	class BackoutValidationResult
+	{
+		public bool Success { get; set; } = true;
+		public String Reason { get; set; } = "";
+	}
+
 	class BackoutCommand : Command
 	{
 		public override string Description => "P4 Admin sanctioned method of backing out a CL";
 
-		public override CustomToolInfo CustomTool => new CustomToolInfo("Backout this CL", "%S") { ShowConsole = true };
+		public override CustomToolInfo CustomTool => new CustomToolInfo("Safe Backout tool", "%S") { ShowConsole = true }; 
+
+		private static async Task<BackoutValidationResult> ValidateBackoutSafety(PerforceConnection Perforce, int Change, ILogger Logger)
+		{
+			BackoutValidationResult Result = new BackoutValidationResult();
+
+			Logger.LogInformation("Examining CL {Change} for Safe Backout", Change);
+
+			// was this change cherry picked using ushell or p4vutil or was it robomerged?
+			DescribeRecord DescribeRec = await Perforce.DescribeAsync(Change, CancellationToken.None);
+			if (DescribeRec.Description.Contains("#ushell-cherrypick", StringComparison.InvariantCultureIgnoreCase)
+			|| DescribeRec.Description.Contains("#p4v-cherrypick", StringComparison.InvariantCultureIgnoreCase)
+			|| DescribeRec.Description.Contains("#ROBOMERGE-SOURCE", StringComparison.InvariantCultureIgnoreCase))
+			{
+				// it was a cherry pick or a robomerged CL, just bail
+				Logger.LogError("CL {Change} did not originate in the current stream,\nto be safe, CLs should be backed out where the change was originally submitted", Change);
+				Result.Success = false;
+				Result.Reason = string.Format($"CL {Change} did not originate in the current stream, to be safe, CLs should be backed out where the change was originally submitted.", Change);
+				return Result;
+			}
+
+			foreach (DescribeFileRecord DescFileRec in DescribeRec.Files)
+			{
+				// is this an edit/add/delete operation?
+				if (!P4ActionGroups.EditActions.Contains(DescFileRec.Action))
+				{
+					Logger.LogError("In CL {Change},\nthe action for the file listed below isn't considered an 'edit',\ntherefore in isn't safe to back it out.\n\n{FileName}", Change, DescFileRec.DepotFile);
+					Result.Success = false;
+					Result.Reason = $"In CL {Change},\nthe action for the file listed below isn't considered an 'edit', therefore in isn't safe to back it out.\n\n{DescFileRec.DepotFile}";
+					return Result;
+				}
+
+				// is it an versioning file?
+				if (DescFileRec.DepotFile.Contains("ObjectVersion.h", StringComparison.InvariantCultureIgnoreCase)
+					|| DescFileRec.DepotFile.Contains("CustomVersion.h", StringComparison.InvariantCultureIgnoreCase)
+					|| DescFileRec.DepotFile.Contains("/Version.h", StringComparison.InvariantCultureIgnoreCase))
+				{
+					Logger.LogError("In CL {Change}, the file below affects asset versioning,\ntherefore it isn't safe to back it out.\n\n{FileName}", Change, DescFileRec.DepotFile);
+					Result.Success = false;
+					Result.Reason = $"In CL {Change}, the file below affects asset versioning, therefore it isn't safe to back it out.\n\n{DescFileRec.DepotFile}";
+					return Result;
+
+				}
+
+				// if it's a binary file, are we backing out the head revision?
+				if (DescFileRec.Type.Contains("+l", StringComparison.InvariantCultureIgnoreCase))
+				{
+					// check the latest revision of this depot file
+					FStatRecord StatRecord = (await Perforce.FStatAsync(new List<string> { DescFileRec.DepotFile }).ToListAsync()).First();
+					if (StatRecord.HeadRevision != DescFileRec.Revision)
+					{
+						Logger.LogError("\nIn CL {Change}, the file below is a binary file and revision #{Rev} isn't the head revision (#{HeadRev}),\ntherefore it isn't safe to back it out\n\n{FileName}",
+							Change, DescFileRec.Revision, StatRecord.HeadRevision, DescFileRec.DepotFile);
+						Result.Success = false;
+						Result.Reason = $"\nIn CL {Change}, the file below is a binary file and revision #{DescFileRec.Revision} isn't the head revision (#{StatRecord.HeadRevision}), therefore it isn't safe to back it out\n\n{DescFileRec.DepotFile}";
+						return Result;
+					}
+				}
+			}
+
+			Logger.LogInformation("CL {Change} seems safe to backout", Change);
+			return Result;
+		}
+
 
+		
 		public override async Task<int> Execute(string[] Args, IReadOnlyDictionary<string, string> ConfigValues, ILogger Logger)
 		{
 			int Change;
@@ -37,17 +109,77 @@ namespace P4VUtils.Commands
 
 			using PerforceConnection Perforce = new PerforceConnection(null, null, Logger);
 
-			// Create a new CL
+			BackoutValidationResult ValidationResult = await ValidateBackoutSafety(Perforce, Change, Logger);
+
+			if (ValidationResult.Success == false)
+			{
+				Logger.LogInformation("\r\nCL {Change} isn't safe to backout, please confirm if you want to proceed.\r\n", Change);
+
+				// warn user
+				MessageBoxResult result = MessageBox.Show(
+					"This backout is potentially unsafe:\r\n" +
+					"\r\n" +
+					ValidationResult.Reason +
+					"\r\n\r\n" +
+					"Are you sure you want to proceed with the backout operation?\r\n" +
+					"\r\n" +
+					"You can tag @p4backouthelp in #ue-build-health or #fn-build-health for assistance.\r\n" +
+					"\r\n",
+					"Unsafe backout detected",
+					MessageBoxButton.YesNo);
+
+				if (result == MessageBoxResult.No)
+				{
+					Logger.LogInformation("Operation canceled.");
+					return 0;
+				}
+				else
+				{
+					Logger.LogInformation("Proceeding with backout operation.");
+				}
+			}
+			
 			DescribeRecord ExistingChangeRecord = await Perforce.DescribeAsync(Change, CancellationToken.None);
 
+			// P4V undo checks for files opened before executing 'undo'
+			List<OpenedRecord> OpenedRecords = await Perforce.OpenedAsync(OpenedOptions.AllWorkspaces | OpenedOptions.ShortOutput, ExistingChangeRecord.Files.Select(x => x.DepotFile).ToArray(), CancellationToken.None).ToListAsync();
+			if (OpenedRecords.Count > 0)
+			{
+				string FileList = String.Join("\r\n", OpenedRecords.Select(x => x.DepotFile));
+
+				Logger.LogInformation("\r\nSome files are checked out on another workspace, please confirm that you wish to continue.\r\n");
+
+				// prompt
+				MessageBoxResult result = MessageBox.Show(
+					"The following files are checked out on another workspace:\r\n" +
+					"\r\n" +
+					FileList +
+					"\r\n\r\n" +
+					"Do you want to proceed with the backout operation?\r\n" +
+					"\r\n",
+					"Warning",
+					MessageBoxButton.YesNo);
+
+				if (result == MessageBoxResult.No)
+				{
+					Logger.LogInformation("Operation canceled.");
+					return 0;
+				}
+				else
+				{
+					Logger.LogInformation("Proceeding with backout operation.");
+				}
+			}
+
 			InfoRecord Info = await Perforce.GetInfoAsync(InfoOptions.None, CancellationToken.None);
 
+			// Create a new CL
 			ChangeRecord NewChangeRecord = new ChangeRecord();
 			NewChangeRecord.User = Info.UserName;
 			NewChangeRecord.Client = Info.ClientName;
 			NewChangeRecord.Description = $"[Backout] - CL{Change}\n#fyi {ExistingChangeRecord.User}\nOriginal CL Desc\n-----------------------------------------------------------------\n{ExistingChangeRecord.Description.TrimEnd()}\n";
 			NewChangeRecord = await Perforce.CreateChangeAsync(NewChangeRecord, CancellationToken.None);
-			
+
 			Logger.LogInformation("Created pending changelist {Change}", NewChangeRecord.Number);
 
 			// Undo the passed in CL
@@ -59,17 +191,19 @@ namespace P4VUtils.Commands
 			// if the original CL and the new CL differ in file count then an error occurs, abort and clean up
 			if (RefreshNewRecord.Files.Count != ExistingChangeRecord.Files.Count)
 			{
-				Logger.LogError("Undo on CL {Change} failed", Change);
+				Logger.LogError("Undo on CL {Change} failed (probably due to an exclusive check out)", Change);
 				foreach (PerforceResponse Response in UndoResponses)
 				{
-					Logger.LogError("  {Response}", Response.ToString())	;
+					Logger.LogError("  {Response}", Response.ToString());
 				}
 
 				// revert files in the new CL
-				if (RefreshNewRecord.Files.Count > 0)
+				List<OpenedRecord> NewCLOpenedRecords = await Perforce.OpenedAsync(OpenedOptions.None, NewChangeRecord.Number, null, null, -1, FileSpecList.Any, CancellationToken.None).ToListAsync();
+
+				if (OpenedRecords.Count > 0)
 				{
 					Logger.LogError("  Reverting");
-					await Perforce.RevertAsync(RefreshNewRecord.Number, Info.ClientName, RevertOptions.None, new[] { "//..." }, CancellationToken.None);
+					await Perforce.RevertAsync(NewChangeRecord.Number, null, RevertOptions.DeleteAddedFiles, NewCLOpenedRecords.Select(x => x.ClientFile!).ToArray(), CancellationToken.None);
 				}
 
 				// delete the new CL
diff --git a/Engine/Extras/P4VUtils/Commands/FindLastEditCommand.cs b/Engine/Extras/P4VUtils/Commands/FindLastEditCommand.cs
--- a/Engine/Extras/P4VUtils/Commands/FindLastEditCommand.cs
+++ b/Engine/Extras/P4VUtils/Commands/FindLastEditCommand.cs
@@ -7,6 +7,7 @@ using System.Collections.Generic;
 using System.Linq;
 using System.Threading;
 using System.Threading.Tasks;
+using P4VUtils.Perforce;
 
 namespace P4VUtils.Commands
 {
@@ -27,32 +28,7 @@ namespace P4VUtils.Commands
 				return 1;
 			}
 
-			// these are used to find out which integrate action are the source
-			IntegrateAction[] IntegrateFromActions =
-			{
-				IntegrateAction.BranchFrom,
-				IntegrateAction.MergeFrom,
-				IntegrateAction.MovedFrom,
-				IntegrateAction.CopyFrom,
-				IntegrateAction.DeleteFrom,
-				IntegrateAction.EditFrom,
-				IntegrateAction.AddFrom
-			};
-
-			// consider these 'edit' actions, stop searching
-			FileAction[] EditActions =
-			{
-				FileAction.Add,
-				FileAction.Edit,
-				FileAction.MoveAdd
-			};
 
-			// consider these 'integrate' actions, keep searching
-			FileAction[] IntegrateActions =
-			{
-				FileAction.Integrate,
-				FileAction.Branch,
-			};
 
 			// the first file log record is what we're interested in
 			// because we're not following integrations, it should be the only one anyway
@@ -69,12 +45,12 @@ namespace P4VUtils.Commands
 				}
 
 				// is it an 'integration'?
-				if (IntegrateActions.Contains(RevisionRecord.Action))
+				if (P4ActionGroups.IntegrateActions.Contains(RevisionRecord.Action))
 				{
 					// where did we integrate from?
 					foreach (IntegrationRecord IntegrationRecord in RevisionRecord.Integrations)
 					{
-						if (IntegrateFromActions.Contains(IntegrationRecord.Action))
+						if (P4ActionGroups.IntegrateFromActions.Contains(IntegrationRecord.Action))
 						{
 							// we need to start the search at the integrated revision # otherwise we might loop forever. 
 							string NewFile = IntegrationRecord.OtherFile + "#" + IntegrationRecord.EndRevisionNumber;
@@ -85,7 +61,7 @@ namespace P4VUtils.Commands
 					break;
 				}
 				// is it an 'edit'?
-				else if (EditActions.Contains(RevisionRecord.Action))
+				else if (P4ActionGroups.EditActions.Contains(RevisionRecord.Action))
 				{
 					// we found the perp so stop here
 
diff --git a/Engine/Extras/P4VUtils/P4ActionGroups.cs b/Engine/Extras/P4VUtils/P4ActionGroups.cs
new file mode 100644
--- /dev/null
+++ b/Engine/Extras/P4VUtils/P4ActionGroups.cs
@@ -0,0 +1,39 @@
+// Copyright Epic Games, Inc. All Rights Reserved.
+
+using EpicGames.Perforce;
+
+namespace P4VUtils.Perforce
+{
+	public static class P4ActionGroups
+	{
+		public static readonly IntegrateAction[] IntegrateFromActions =
+		{
+			IntegrateAction.BranchFrom,
+			IntegrateAction.MergeFrom,
+			IntegrateAction.MovedFrom,
+			IntegrateAction.CopyFrom,
+			IntegrateAction.DeleteFrom,
+			IntegrateAction.EditFrom,
+			IntegrateAction.AddFrom
+		};
+
+		public static readonly FileAction[] EditActions =
+		{
+			FileAction.Add,
+			FileAction.Edit,
+			FileAction.Delete
+		};
+
+		public static readonly FileAction[] MoveActions =
+		{
+			FileAction.MoveAdd,
+			FileAction.MoveDelete
+		};
+
+		public static readonly FileAction[] IntegrateActions =
+		{
+			FileAction.Integrate,
+			FileAction.Branch,
+		};
+	}
+}
