P4VUtils: fixed a few issues in the safe backout tool:
- will now report multiple safety concerns if there are > 1
- message box body is now limited to 2048 characters to avoid issue where you can't even see the buttons (all errors are still logged to the console window which has a scroll bar)

#rnx

diff --git a/Engine/Extras/P4VUtils/Commands/BackoutCommand.cs b/Engine/Extras/P4VUtils/Commands/BackoutCommand.cs
--- a/Engine/Extras/P4VUtils/Commands/BackoutCommand.cs
+++ b/Engine/Extras/P4VUtils/Commands/BackoutCommand.cs
@@ -17,15 +17,16 @@ namespace P4VUtils.Commands
 {
 	class BackoutValidationResult
 	{
-		public bool Success { get; set; } = true;
-		public String Reason { get; set; } = "";
+		public List<String> Errors = new List<String>();
 	}
 
 	class BackoutCommand : Command
 	{
 		public override string Description => "P4 Admin sanctioned method of backing out a CL";
 
-		public override CustomToolInfo CustomTool => new CustomToolInfo("Safe Backout tool", "%S") { ShowConsole = true }; 
+		public override CustomToolInfo CustomTool => new CustomToolInfo("Safe Backout tool", "%S") { ShowConsole = true };
+
+		private readonly int MAX_ERROR_LENGTH = 2048;
 
 		private static async Task<BackoutValidationResult> ValidateBackoutSafety(PerforceConnection Perforce, int Change, ILogger Logger)
 		{
@@ -41,9 +42,7 @@ namespace P4VUtils.Commands
 			{
 				// it was a cherry pick or a robomerged CL, just bail
 				Logger.LogError("CL {Change} did not originate in the current stream,\nto be safe, CLs should be backed out where the change was originally submitted", Change);
-				Result.Success = false;
-				Result.Reason = string.Format($"CL {Change} did not originate in the current stream, to be safe, CLs should be backed out where the change was originally submitted.", Change);
-				return Result;
+				Result.Errors.Add(string.Format($"CL {Change} did not originate in the current stream, to be safe, CLs should be backed out where the change was originally submitted.", Change));
 			}
 
 			foreach (DescribeFileRecord DescFileRec in DescribeRec.Files)
@@ -52,9 +51,7 @@ namespace P4VUtils.Commands
 				if (!P4ActionGroups.EditActions.Contains(DescFileRec.Action))
 				{
 					Logger.LogError("In CL {Change},\nthe action for the file listed below isn't considered an 'edit',\ntherefore in isn't safe to back it out.\n\n{FileName}", Change, DescFileRec.DepotFile);
-					Result.Success = false;
-					Result.Reason = $"In CL {Change},\nthe action for the file listed below isn't considered an 'edit', therefore in isn't safe to back it out.\n\n{DescFileRec.DepotFile}";
-					return Result;
+					Result.Errors.Add($"In CL {Change},\nthe action for the file listed below isn't considered an 'edit', therefore in isn't safe to back it out.\n\n{DescFileRec.DepotFile}");
 				}
 
 				// is it an versioning file?
@@ -63,10 +60,7 @@ namespace P4VUtils.Commands
 					|| DescFileRec.DepotFile.Contains("/Version.h", StringComparison.InvariantCultureIgnoreCase))
 				{
 					Logger.LogError("In CL {Change}, the file below affects asset versioning,\ntherefore it isn't safe to back it out.\n\n{FileName}", Change, DescFileRec.DepotFile);
-					Result.Success = false;
-					Result.Reason = $"In CL {Change}, the file below affects asset versioning, therefore it isn't safe to back it out.\n\n{DescFileRec.DepotFile}";
-					return Result;
-
+					Result.Errors.Add($"In CL {Change}, the file below affects asset versioning, therefore it isn't safe to back it out.\n\n{DescFileRec.DepotFile}");
 				}
 
 				// if it's a binary file, are we backing out the head revision?
@@ -78,14 +72,15 @@ namespace P4VUtils.Commands
 					{
 						Logger.LogError("\nIn CL {Change}, the file below is a binary file and revision #{Rev} isn't the head revision (#{HeadRev}),\ntherefore it isn't safe to back it out\n\n{FileName}",
 							Change, DescFileRec.Revision, StatRecord.HeadRevision, DescFileRec.DepotFile);
-						Result.Success = false;
-						Result.Reason = $"\nIn CL {Change}, the file below is a binary file and revision #{DescFileRec.Revision} isn't the head revision (#{StatRecord.HeadRevision}), therefore it isn't safe to back it out\n\n{DescFileRec.DepotFile}";
-						return Result;
+						Result.Errors.Add($"\nIn CL {Change}, the file below is a binary file and revision #{DescFileRec.Revision} isn't the head revision (#{StatRecord.HeadRevision}), therefore it isn't safe to back it out\n\n{DescFileRec.DepotFile}");
 					}
 				}
 			}
 
-			Logger.LogInformation("CL {Change} seems safe to backout", Change);
+			if (Result.Errors.Count == 0)
+			{
+				Logger.LogInformation("CL {Change} seems safe to backout", Change);
+			}
 			return Result;
 		}
 
@@ -111,15 +106,22 @@ namespace P4VUtils.Commands
 
 			BackoutValidationResult ValidationResult = await ValidateBackoutSafety(Perforce, Change, Logger);
 
-			if (ValidationResult.Success == false)
+			if (ValidationResult.Errors.Count > 0)
 			{
 				Logger.LogInformation("\r\nCL {Change} isn't safe to backout, please confirm if you want to proceed.\r\n", Change);
 
+				string ErrorString = String.Join("\r\n", ValidationResult.Errors); 
+				if (ErrorString.Length > MAX_ERROR_LENGTH)
+				{
+					ErrorString = ErrorString.Substring(0, MAX_ERROR_LENGTH);
+					ErrorString += "\r\n (...)";
+				}
+
 				// warn user
 				MessageBoxResult result = MessageBox.Show(
 					"This backout is potentially unsafe:\r\n" +
 					"\r\n" +
-					ValidationResult.Reason +
+					ErrorString +
 					"\r\n\r\n" +
 					"Are you sure you want to proceed with the backout operation?\r\n" +
 					"\r\n" +
@@ -130,12 +132,12 @@ namespace P4VUtils.Commands
 
 				if (result == MessageBoxResult.No)
 				{
-					Logger.LogInformation("Operation canceled.");
+					Logger.LogInformation("\r\nOperation canceled.");
 					return 0;
 				}
 				else
 				{
-					Logger.LogInformation("Proceeding with backout operation.");
+					Logger.LogInformation("\r\nProceeding with backout operation.");
 				}
 			}
 			
@@ -145,15 +147,24 @@ namespace P4VUtils.Commands
 			List<OpenedRecord> OpenedRecords = await Perforce.OpenedAsync(OpenedOptions.AllWorkspaces | OpenedOptions.ShortOutput, ExistingChangeRecord.Files.Select(x => x.DepotFile).ToArray(), CancellationToken.None).ToListAsync();
 			if (OpenedRecords.Count > 0)
 			{
-				string FileList = String.Join("\r\n", OpenedRecords.Select(x => x.DepotFile));
+				HashSet<string> UniqueDepotFiles = (OpenedRecords.Select(x => x.DepotFile)).ToHashSet();
+				string FileListString = string.Join("\r\n", UniqueDepotFiles);
 
 				Logger.LogInformation("\r\nSome files are checked out on another workspace, please confirm that you wish to continue.\r\n");
+				Logger.LogInformation("{FileList}", FileListString);
+
+				string FileListTruncated = FileListString;
+				if (FileListString.Length > MAX_ERROR_LENGTH)
+				{
+					FileListTruncated = FileListTruncated.Substring(0, MAX_ERROR_LENGTH);
+					FileListTruncated += "(...)";
+				}
 
 				// prompt
 				MessageBoxResult result = MessageBox.Show(
 					"The following files are checked out on another workspace:\r\n" +
 					"\r\n" +
-					FileList +
+					FileListTruncated +
 					"\r\n\r\n" +
 					"Do you want to proceed with the backout operation?\r\n" +
 					"\r\n",
@@ -162,12 +173,12 @@ namespace P4VUtils.Commands
 
 				if (result == MessageBoxResult.No)
 				{
-					Logger.LogInformation("Operation canceled.");
+					Logger.LogInformation("\r\nOperation canceled.");
 					return 0;
 				}
 				else
 				{
-					Logger.LogInformation("Proceeding with backout operation.");
+					Logger.LogInformation("\r\nProceeding with backout operation.");
 				}
 			}
 
