[P4VUtils] - re-organizing tools
Renamed UE Helpers to UE Toolbox
Moved backout tool to UE Toolbox
Moved 'find last edit' to 'Root Helpers' (from UE Helpers)
#rnx

diff --git a/Engine/Extras/P4VUtils/Program.cs b/Engine/Extras/P4VUtils/Program.cs
--- a/Engine/Extras/P4VUtils/Program.cs
+++ b/Engine/Extras/P4VUtils/Program.cs
@@ -49,6 +49,8 @@ namespace P4VUtils
 		{
 			["describe"] = new DescribeCommand(),
 			["copyclnum"] = new CopyCLCommand(),
+			["findlastedit"] = new FindLastEditCommand(),
+			["findlasteditbyline"] = new P4BlameCommand(),
 		};
 
 		// UESubmit - commands that help with submitting files/changelists
@@ -58,10 +60,9 @@ namespace P4VUtils
 		};
 
 		// UEHelpers - commands that help with common but simple operations
-		public static IReadOnlyDictionary<string, Command> HelperCommands { get; } = new Dictionary<string, Command>(StringComparer.OrdinalIgnoreCase)
+		public static IReadOnlyDictionary<string, Command> ToolboxCommands { get; } = new Dictionary<string, Command>(StringComparer.OrdinalIgnoreCase)
 		{
-			["findlastedit"] = new FindLastEditCommand(),
-			["findlasteditbyline"] = new P4BlameCommand(),
+			["backout"] = new BackoutCommand(),
 			["snapshot"] = new SnapshotCommand(),
 			["reconcilecode"] = new FastReconcileCodeEditsCommand(),
 			["reconcileall"] = new FastReconcileAllEditsCommand(),
@@ -77,7 +78,6 @@ namespace P4VUtils
 			["cherrypick"] = new CherryPickCommand(),
 			["converttoedit"] = new ConvertToEditCommand(),
 			["edigrate"] = new EdigrateCommand(),
-			["backout"] = new BackoutCommand(),
 		};
 
 		// UEHorde Folder - local build and horde preflights
@@ -89,7 +89,7 @@ namespace P4VUtils
 			["movewriteablepreflightandsubmit"] = new MoveWriteableFilesthenPreflightAndSubmitCommand(),
 		};
 
-		public static IDictionary<string, Command> Commands = SubmissionCommands.Concat(RootHelperCommands).Concat(HelperCommands).Concat(IntegrateCommands).Concat(HordeCommands).ToDictionary(p => p.Key, p => p.Value, StringComparer.OrdinalIgnoreCase);
+		public static IDictionary<string, Command> Commands = SubmissionCommands.Concat(RootHelperCommands).Concat(ToolboxCommands).Concat(IntegrateCommands).Concat(HordeCommands).ToDictionary(p => p.Key, p => p.Value, StringComparer.OrdinalIgnoreCase);
 
 		static void PrintHelp(ILogger Logger)
 		{
@@ -446,11 +446,11 @@ namespace P4VUtils
 			// Insert new entries
 			if (bInstall)
 			{
-				InstallCommandsListInFolder("UERootHelpers", false/*AddFolderToContextMenu*/, RootHelperCommands, Document, DotNetLocation, AssemblyLocation, Logger);
-				InstallCommandsListInFolder("UESubmit", true/*AddFolderToContextMenu*/, SubmissionCommands, Document, DotNetLocation, AssemblyLocation, Logger);
-				InstallCommandsListInFolder("UEHelpers", true/*AddFolderToContextMenu*/, HelperCommands, Document, DotNetLocation, AssemblyLocation, Logger);
-				InstallCommandsListInFolder("UEIntegrate", true/*AddFolderToContextMenu*/, IntegrateCommands, Document, DotNetLocation, AssemblyLocation, Logger);
-				InstallCommandsListInFolder("UEHorde", true/*AddFolderToContextMenu*/, HordeCommands, Document, DotNetLocation, AssemblyLocation, Logger);
+				InstallCommandsListInFolder("UE RootHelpers", false/*AddFolderToContextMenu*/, RootHelperCommands, Document, DotNetLocation, AssemblyLocation, Logger);
+				InstallCommandsListInFolder("UE Submit", true/*AddFolderToContextMenu*/, SubmissionCommands, Document, DotNetLocation, AssemblyLocation, Logger);
+				InstallCommandsListInFolder("UE Toolbox", true/*AddFolderToContextMenu*/, ToolboxCommands, Document, DotNetLocation, AssemblyLocation, Logger);
+				InstallCommandsListInFolder("UE Integrate", true/*AddFolderToContextMenu*/, IntegrateCommands, Document, DotNetLocation, AssemblyLocation, Logger);
+				InstallCommandsListInFolder("UE Horde", true/*AddFolderToContextMenu*/, HordeCommands, Document, DotNetLocation, AssemblyLocation, Logger);
 			}
 
 			// Save the new document
