[P4VUtils] - fix minor issue in Submit command

#rnx

diff --git a/Engine/Extras/P4VUtils/Commands/SubmitCommand.cs b/Engine/Extras/P4VUtils/Commands/SubmitCommand.cs
--- a/Engine/Extras/P4VUtils/Commands/SubmitCommand.cs
+++ b/Engine/Extras/P4VUtils/Commands/SubmitCommand.cs
@@ -734,7 +734,7 @@ namespace P4VUtils.Commands
 			if (!String.IsNullOrEmpty(engineIdentifier))
 			{
 				string engineRoot = GetEngineRootDirFromIdentifier(engineIdentifier, logger);
-				if (!String.IsNullOrEmpty(engineIdentifier))
+				if (!String.IsNullOrEmpty(engineRoot))
 				{
 					return engineRoot;
 				}
