[P4VUtils] - add DescribeDirectory tool to output a list of all CL's in the last 30 days for a specific folder in a succinct/clean format.

#rnx

diff --git a/Engine/Extras/P4VUtils/Commands/DescribeCommand.cs b/Engine/Extras/P4VUtils/Commands/DescribeCommand.cs
--- a/Engine/Extras/P4VUtils/Commands/DescribeCommand.cs
+++ b/Engine/Extras/P4VUtils/Commands/DescribeCommand.cs
@@ -6,9 +6,12 @@ using Microsoft.Extensions.Logging;
 using System;
 using System.Collections.Generic;
 using System.ComponentModel;
+using System.IO;
 using System.Text;
 using System.Threading;
 using System.Threading.Tasks;
+using System.Runtime.InteropServices;
+using System.Diagnostics;
 
 namespace P4VUtils.Commands
 {
@@ -43,4 +46,147 @@ namespace P4VUtils.Commands
 			return 0;
 		}
 	}
+
+	class DescribeDirectoryCommand : Command
+	{
+		public override string Description => "Prints out a description of last 30 days of CL's in the given directory to txt and csv";
+
+		public override CustomToolInfo CustomTool => new CustomToolInfo("Describe Directory", "%D") { ShowConsole = true, RefreshUI = false };
+
+		private int MaxDaysHistory = 30;
+		private int MaxChangeListCount = 10000;
+
+		static void OSOpenFile(string FilePath)
+		{
+			if (RuntimeInformation.IsOSPlatform(OSPlatform.Windows))
+			{
+				Process.Start(new ProcessStartInfo(FilePath) { UseShellExecute = true });
+			}
+			else if (RuntimeInformation.IsOSPlatform(OSPlatform.Linux))
+			{
+				Process.Start("xdg-open", FilePath);
+			}
+			else if (RuntimeInformation.IsOSPlatform(OSPlatform.OSX))
+			{
+				Process.Start("open", FilePath);
+			}
+		}
+
+		public override async Task<int> Execute(string[] Args, IReadOnlyDictionary<string, string> ConfigValues, ILogger Logger)
+		{
+			if (Args.Length < 2)
+			{
+				Logger.LogError("Missing directory path(s)");
+				return 1;
+			}
+
+			List<string> DescribePathsStrings = new List<string>();
+
+			// %D is 'selected files or folders', ensure we only use folders
+			for (int i = 1; i < Args.Length; i++)
+			{
+				Logger.LogInformation("params {Idx} - {Arg}", i, Args[i]);
+
+				if (Args[i].EndsWith("...", StringComparison.Ordinal))
+				{
+					DescribePathsStrings.Add(Args[i]);
+				}
+				else
+				{
+					Logger.LogError("   ignoring individial selected file {File}", Args[i]);
+				}
+			}
+
+			using PerforceConnection Perforce = new PerforceConnection(null, null, Logger);
+
+			// grab a list of the changes with a large max count
+			List<ChangesRecord> ChangeRecords = await Perforce.GetChangesAsync(ChangesOptions.LongOutput, MaxChangeListCount, ChangeStatus.Submitted, DescribePathsStrings, CancellationToken.None);
+
+			if (ChangeRecords.Count == 0)
+			{
+				Logger.LogWarning("Could not find any p4 change record for {File}", DescribePathsStrings);
+				return 1;
+			}
+
+			StringBuilder CSVFormat = new StringBuilder();
+			StringBuilder PlainTextFormat = new StringBuilder();
+
+			DateTime TimeNow = DateTime.Now;
+			// iterate over the changes, filter to within the last x days
+			foreach (ChangesRecord ChangeRecord in ChangeRecords)
+			{
+				TimeSpan Delta = TimeNow - ChangeRecord.Time;
+				if (Delta.TotalDays < MaxDaysHistory)
+				{
+					// grab the description string - split it and remove lines that start with # or trim lines that have # within them.
+					string[] Lines = ChangeRecord.Description.Split('\n');
+
+					StringBuilder FilteredOutput = new StringBuilder();
+
+					foreach(string line in Lines)
+					{
+						string[] subLines = line.Split('#');
+						if (subLines.Length > 0)
+						{
+							FilteredOutput.Append(subLines[0]);
+						}
+					}
+
+					// now remove all newlines for the plain text output, truncate to 100 chars for text size reasons
+					string Single = FilteredOutput.ToString().
+						Replace("\n", "", System.StringComparison.Ordinal).
+						Replace("\r", "", System.StringComparison.Ordinal);
+					string FinalSingleString;
+
+					if (Single.Length > 100)
+					{
+						FinalSingleString = Single.Substring(0, 100);
+						FinalSingleString += "(...)";
+					}
+					else
+					{
+						FinalSingleString = Single;
+					}
+
+					if(FinalSingleString.Length < 2)
+					{
+						FinalSingleString = Lines[0];
+					}
+
+					PlainTextFormat.AppendLine("{0} - {1} - {2} - {3}", ChangeRecord.Number, ChangeRecord.User, FinalSingleString, ChangeRecord.Time.ToString());
+
+					string CSVString = FilteredOutput.ToString().Replace("\"", "\"\"",StringComparison.Ordinal);
+
+					if (CSVString.Length < 2)
+					{
+						CSVString = Lines[0];
+					}
+
+					CSVFormat.AppendLine("{0},{1},\"{2}\",{3}", ChangeRecord.Number, ChangeRecord.User, CSVString, ChangeRecord.Time.ToString());
+				}
+			}
+
+			string GuidString = Guid.NewGuid().ToString();
+			string CsvPath = Path.Combine(Path.GetTempPath(), GuidString + ".csv");
+			string PlainTextPath = Path.Combine(Path.GetTempPath(), GuidString + ".txt");
+
+			using (StreamWriter Writer = new StreamWriter(PlainTextPath))
+			{
+				await Writer.WriteLineAsync(PlainTextFormat.ToString());
+			}
+
+			using (StreamWriter Writer = new StreamWriter(CsvPath))
+			{
+				await Writer.WriteLineAsync(CSVFormat.ToString());
+			}
+
+			OSOpenFile(PlainTextPath);
+
+			Logger.LogInformation("Complete");
+			Logger.LogInformation("PlainText: {PlainText}", PlainTextPath);
+			Logger.LogInformation("CSV Form : {CsvPath}", CsvPath);
+
+			return 0;
+		}
+	}
 }
diff --git a/Engine/Extras/P4VUtils/Program.cs b/Engine/Extras/P4VUtils/Program.cs
--- a/Engine/Extras/P4VUtils/Program.cs
+++ b/Engine/Extras/P4VUtils/Program.cs
@@ -70,6 +70,7 @@ namespace P4VUtils
 			["unshelvemakedatawritable"] = new UnshelveMakeDataWritable(),
 			["convertcldatatolocalwritable"] = new ConvertCLDataToLocalWritable(),
 			["convertdatatolocalwritable"] = new ConvertDataToLocalWritable(),
+			["describedirectory"] = new DescribeDirectoryCommand(),
 		};
 
 		// UEIntegrate Folder commands - complex commands to facilitate integrations/backout
