Fixed an issue where P4VUtils was not correctly finding the precompiled binary version (PCB) of UnrealVirtualizationTool if needed.

#rnx

- Previously we were checking the version of the tool (if found) and assuming that a promoted build meant a precompiled binary. Recent changed to our build farm means that this assumption no longer works and was preventing people who did not have source code synced from running the tool.
- The previous logic was to check if the editor was a PCB and if it was not then compile the code locally. So if a user downloaded PCB we would use them, up until they recompiled the editor locally at which point we'd change to compiling UnrealVirtualizationTool as well, which was an attempt to keep things in sync. However once we stopped detecting PCB properly, every editor was considered "locally built" meaning we had to build the tool and this would prevent users without source code (who could not build) from continuing.

- As a short term fix the logic has been changed. We now check to see if the UBT project can be found. If it can we assume the user has source code in which case we just build the tool. If we do not find the project we assume no source code and the user must be using PCB, so we try to use that.
- This does mean that a user with source code could get a mismatch in future, using the PCB editor with a locally built UnrealVirtualizationTool (that might have changes compared to the editor) but that case should be less common and less of a problem than our current issues.
- This logic change can be reverted once a method is deployed that will allow us to more reliably tell the difference between PCB exes and locally built.

diff --git a/Engine/Extras/P4VUtils/Commands/SubmitCommand.cs b/Engine/Extras/P4VUtils/Commands/SubmitCommand.cs
--- a/Engine/Extras/P4VUtils/Commands/SubmitCommand.cs
+++ b/Engine/Extras/P4VUtils/Commands/SubmitCommand.cs
@@ -271,108 +271,74 @@ namespace P4VUtils.Commands
 		private static async Task<bool> ValidateOrBuildVirtualizationTool(string engineRoot, ILogger logger)
 		{
 			// Check to see if the precompiled binaries for the  editor and virtualization tool are available
-			bool hasEditorPCB = IsPrecompiledVersionAvailable(engineRoot, "UnrealEditor", logger);
-			bool hasToolPCB = IsPrecompiledVersionAvailable(engineRoot, "UnrealVirtualizationTool", logger);
-
-			// If the editor is not a precompiled binary, then we know that the user has compiled it locally.
-			// This means that even if we have a precompiled version of UnrealVirtualizationTool we cannot
-			// use it as it's virtualization process might be 'out of date' compared to the locally compiled
-			// editor.
-			// This means we can only use the precompiled binary for UnrealVirtualizationTool if UnrealEditor
-			// is also a precompiled binary.
-			if (hasEditorPCB && hasToolPCB)
-			{
-				logger.LogInformation("Using a precompiled binary version of UnrealVirtualizationTool");
-				return true;
-			}
-			else if (hasToolPCB)
+
+			if (!IsSourceCodeAvaliable(engineRoot))
 			{
-				// Let the user know why we are compiling UnrealVirtualizationTool even though there is a precompiled
-				// version available.
-				logger.LogInformation("Cannot use the precompiled version of UnrealVirtualizationTool as the UnrealEditor has been locally compiled");
+				if (DoesEngineToolExist(engineRoot, "UnrealVirtualizationTool"))
+				{
+					logger.LogInformation("Using a precompiled binary version of UnrealVirtualizationTool");
+					return true;
+				}
+				else
+				{
+					logger.LogError("No source code and no precompiled binary of UnrealVirtualizationTool found");
+					return false;
+				}
 			}
+			else
+			{
+				// Since we have no good way to determine if an exe was from the PCB or not, if we have
+				// source code we will just have to build the tool anyway for safety.
+				// Once this is fixed we should use the old logic which was to use the PCB version of
+				// UnrealVirtualizationTool if the editor was a PCB version and only to recompile it 
+				// once we detect that the editor has been compiled locally.
 
-			// We must try to build the tool locally, if the user does not have a valid code compilation 
-			// tool chain installed then UnrealBuildTool will give them errors.
-			return await BuildVirtualizationTool(engineRoot, logger);
+				// We must try to build the tool locally, if the user does not have a valid code compilation 
+				// tool chain installed then UnrealBuildTool will give them errors.
+				return await BuildVirtualizationTool(engineRoot, logger);
+			}
 		}
 
 		/// <summary>
-		/// Checks if there is a precompiled version of a given program available for use.
+		/// Checks to see if the user has source code synced or not.
+		/// To make this assumption we check to see if we can find the UnrealBuildTool project. If we can
+		/// then it is a fair bet that the user can compile code, if they do not have the project then
+		/// they probably can't.
 		/// </summary>
 		/// <param name="engineRoot">Root path of the engine we want to build the tool for</param>
-		/// <param name="programName">Name of the program to check for</param>
-		/// <param name="logger">Interface for logging</param>
-		/// <returns>True if there is a version of the program that was precompiled, otherwise false</returns>
-		private static bool IsPrecompiledVersionAvailable(string engineRoot, string programName, ILogger logger)
+		/// <returns>True if source code is present, otherwise false</returns>
+		private static bool IsSourceCodeAvaliable(string engineRoot)
 		{
-			// Depending on how a program is set up it might have a .version file or a .target file that contains info
-			// about the exe itself. From this we can find a JSON property "IsPromotedBuild" which when none zero 
-			// means that the build was precompiled.
-			// This isn't the most robust check but should work for now until we can add an 'offical' codified way to
-			// check for precompiled binaries.
-
-			// First we check for a .version file for the given program
-			string versionPath = String.Format(@"{0}\Engine\Binaries\Win64\{1}.version", engineRoot, programName);
-			
-			if (System.IO.File.Exists(versionPath))
-			{ 
-				using (StreamReader inputStream = new StreamReader(versionPath))
-				{
-					string jsonObject = inputStream.ReadToEnd();
-
-					BuildVersion? buildVersion = null;
-					try
-					{
-						buildVersion = JsonSerializer.Deserialize<BuildVersion>(jsonObject);
-					}
-					catch (System.Text.Json.JsonException)
-					{
-						// No need to print the exception as the user will not be expected to fix the file, a warning will be shown below
-					}
+			string versionPath = String.Format(@"{0}\Engine\Source\Programs\UnrealBuildTool\UnrealBuildTool.csproj", engineRoot);
 
-					if (buildVersion == null || buildVersion.IsPromotedBuild == null)
-					{
-						logger.LogWarning("Unable to read build version from '{FilePath}'", versionPath);
-					}
-					else if (buildVersion.IsPromotedBuild != 0)
-					{
-						return true;
-					}
-				}
+			if (System.IO.File.Exists(versionPath))
+			{
+				return true;
 			}
+			else
+			{
+				return false;
+			}
+		}
 
-			// We failed to find what we needed from the .version file, so try the .target file instead
-			string targetPath = String.Format(@"{0}\Engine\Binaries\Win64\{1}.target", engineRoot, programName);
+		/// <summary>
+		/// Checks if the given tool exists in the engine binaries or not
+		/// </summary>
+		/// <param name="engineRoot">Root path of the engine we want to build the tool for</param>
+		/// <param name="toolName">The name of the tool to look for</param>
+		/// <returns>TRue if the exe for the tool already exists, otherwise false</returns>
+		private static bool DoesEngineToolExist(string engineRoot, string toolName)
+		{
+			string toolPath = String.Format(@"{0}\Engine\Binaries\Win64\{1}.exe", engineRoot, toolName);
 
-			if (System.IO.File.Exists(targetPath))
+			if (System.IO.File.Exists(toolPath))
 			{
-				using (StreamReader inputStream = new StreamReader(targetPath))
-				{
-					string jsonObject = inputStream.ReadToEnd();
-
-					TargetBuildVersion? targetBuildVersion = null;
-					try
-					{
-						targetBuildVersion = JsonSerializer.Deserialize<TargetBuildVersion>(jsonObject);
-					}
-					catch (System.Text.Json.JsonException)
-					{
-						// No need to print the exception as the user will not be expected to fix the file, a warning will be shown below
-					}
-					
-					if (targetBuildVersion == null || targetBuildVersion.Version == null || targetBuildVersion.Version.IsPromotedBuild == null)
-					{
-						logger.LogWarning("Unable to read build version from '{FilePath}'", targetPath);
-					}
-					else if (targetBuildVersion.Version.IsPromotedBuild != 0)
-					{
-						return true;
-					}
-				}
+				return true;
+			}
+			else
+			{
+				return false;
 			}
-
-			return false;
 		}
 
 		/// <summary>
