Logging clean up pass

#rnx

- Prefixed all error logs with "Error: " as the p4v terminal does not support text colours so there was no way to tell info from an error.
- Reworded some logging info and added a few more linebreak to improve readability.

diff --git a/Engine/Extras/P4VUtils/Commands/SubmitCommand.cs b/Engine/Extras/P4VUtils/Commands/SubmitCommand.cs
--- a/Engine/Extras/P4VUtils/Commands/SubmitCommand.cs
+++ b/Engine/Extras/P4VUtils/Commands/SubmitCommand.cs
@@ -69,7 +69,7 @@ namespace P4VUtils.Commands
 			// Parse command lines
 			if (args.Length < 3)
 			{
-				logger.LogError("Not enough args for command, tool is now exiting");
+				logger.LogError("Error: Not enough args for command, tool is now exiting");
 				return 1;
 			}
 
@@ -81,7 +81,7 @@ namespace P4VUtils.Commands
 			{
 				if (!int.TryParse(args[index], out int changeNumber))
 				{
-					logger.LogError("'{Argument}' is not a numbered changelist, tool is now exiting", args[index]);
+					logger.LogError("Error: '{Argument}' is not a numbered changelist, tool is now exiting", args[index]);
 					return 1;
 				}
 
@@ -90,7 +90,7 @@ namespace P4VUtils.Commands
 
 			if (changelistsToSubmit.Count == 0)
 			{
-				logger.LogError("No changelists to submit were provided, tool is now exiting");
+				logger.LogError("Error: No changelists to submit were provided, tool is now exiting");
 				return 1;
 			}
 
@@ -104,7 +104,7 @@ namespace P4VUtils.Commands
 			using IPerforceConnection perforceConnection = await PerforceConnection.CreateAsync(settings, logger);
 			if (perforceConnection == null)
 			{
-				logger.LogError("Failed to connect to Perforce, tool is now exiting");
+				logger.LogError("Error: Failed to connect to Perforce, tool is now exiting");
 				return 1;
 			}
 
@@ -116,7 +116,7 @@ namespace P4VUtils.Commands
 				int changeNumber = changelistsToSubmit[index];
 				if(await ProcessChangelist(perforceConnection, changeNumber, logger) == false)
 				{
-					logger.LogError("Failed to process {Changelist}, tool is now exiting", changeNumber);
+					logger.LogError("\nError: Failed to process CL {Changelist}, tool is now exiting", changeNumber);
 					return 1;
 				}
 			}
@@ -131,7 +131,7 @@ namespace P4VUtils.Commands
 		{
 			if (perforceConnection == null || perforceConnection.Settings == null || perforceConnection.Settings.ClientName == null)
 			{
-				logger.LogError("Invalid Perforce connection!");
+				logger.LogError("Error: Invalid Perforce connection!");
 				return false;
 			}
 
@@ -154,7 +154,7 @@ namespace P4VUtils.Commands
 				// So we should check for shelved files now and early out before the virtualization process runs.
 				if (await DoesChangelistHaveShelvedFiles(perforceConnection, changeNumber) == true)
 				{
-					logger.LogError("Changelist {Change} has shelved files and cannot be submitted", changeNumber);
+					logger.LogError("Error: Changelist {Change} has shelved files and cannot be submitted", changeNumber);
 					return false;
 				}
 
@@ -176,7 +176,7 @@ namespace P4VUtils.Commands
 					if (!String.IsNullOrEmpty(engineRoot))
 					{
 						logger.NewLine();
-						logger.LogInformation("Attempting to virtualize packages in project '{Project}' via the engine installation '{Engine}'", project.Key, engineRoot);
+						logger.LogInformation("Attempting to virtualize packages in project '{Project}' using the engine installation '{Engine}'", project.Key, engineRoot);
 						// @todo Many projects can share the same engine install, and technically UnrealVirtualizationTool
 						// supports the virtualization files from many projects at the same time. We could consider doing
 						// this pass per engine install rather than per project? At the very least we should only 'build'
@@ -208,7 +208,7 @@ namespace P4VUtils.Commands
 					}
 					else
 					{
-						logger.LogError("Failed to find engine root for project {Project}", project.Key);
+						logger.LogError("Error: Failed to find engine root for project {Project}", project.Key);
 						return false;
 					}
 				}
@@ -235,7 +235,7 @@ namespace P4VUtils.Commands
 					PerforceResponse updateResponse = await perforceConnection.TryUpdateChangeAsync(UpdateChangeOptions.None, changeRecord, CancellationToken.None);
 					if (!updateResponse.Succeeded)
 					{
-						logger.LogError("Failed to remove the virtualization tags!");
+						logger.LogError("Error: Failed to remove the virtualization tags!");
 					}
 					logger.LogInformation("Virtualization tags have been removed.");
 
@@ -281,7 +281,7 @@ namespace P4VUtils.Commands
 				}
 				else
 				{
-					logger.LogError("No source code and no precompiled binary of UnrealVirtualizationTool found");
+					logger.LogError("Error: No source code and no precompiled binary of UnrealVirtualizationTool found");
 					return false;
 				}
 			}
@@ -378,7 +378,7 @@ namespace P4VUtils.Commands
 						await bufferedOutput.CopyToAsync(stdOutput, CancellationToken.None);
 					}
 
-					logger.LogError("Failed to build UnrealVirtualizationTool");
+					logger.LogError("Error: Failed to build UnrealVirtualizationTool");
 					return false;
 				}
 			}
@@ -410,7 +410,7 @@ namespace P4VUtils.Commands
 
 				if (Process.ExitCode != 0)
 				{
-					logger.LogError("UnrealVirtualizationTool failed!");
+					logger.LogError("Error: UnrealVirtualizationTool failed!");
 					return false;
 				}
 			}
@@ -504,7 +504,7 @@ namespace P4VUtils.Commands
 			else
 			{
 				// Log every response that was a failure and has an error message associated with it
-				logger.LogError("ERROR - Submit failed due to:");
+				logger.LogError("Error: Submit failed due to:");
 				foreach (PerforceResponse<SubmitRecord> response in submitResponses)
 				{
 					if (response.Failed && response.Error != null)
@@ -536,7 +536,7 @@ namespace P4VUtils.Commands
 			}
 			catch (Exception)
 			{
-				logger.LogError("Failed to get the description {Change} so we can edit it", changeNumber);
+				logger.LogError("Error: Failed to get the description {Change} so we can edit it", changeNumber);
 				return null;
 			}
 
@@ -554,7 +554,7 @@ namespace P4VUtils.Commands
 			}
 			else
 			{
-				logger.LogError("Failed to edit the description of {Change} due to\n{Message}", changeNumber, updateResponse.Error!.ToString());
+				logger.LogError("Error: Failed to edit the description of {Change} due to\n{Message}", changeNumber, updateResponse.Error!.ToString());
 				return null;
 			}
 		}
@@ -576,13 +576,13 @@ namespace P4VUtils.Commands
 			}
 			catch (Exception)
 			{
-				logger.LogError("Failed to find changelist {Change}", changeNumber);
+				logger.LogError("Error: Failed to find changelist {Change}", changeNumber);
 				return null;
 			}
 
 			if (changeRecord.Files.Count == 0)
 			{
-				logger.LogError("Changelist {Change} is empty, cannot submit", changeNumber);
+				logger.LogError("Error: Changelist {Change} is empty, cannot submit", changeNumber);
 				return null;
 			}
 
@@ -777,7 +777,7 @@ namespace P4VUtils.Commands
 			}
 			catch (Exception ex)
 			{
-				logger.LogError("Failed to parse {File} to find the engine association due to: {Reason}", projectFilePath, ex.Message);
+				logger.LogError("Error: Failed to parse {File} to find the engine association due to: {Reason}", projectFilePath, ex.Message);
 			}
 
 			// @todo In the native version if there is no identifier we will try to
