Display the final submitted changelist number as well as the original changelist number.

- When this tool was written the p4 c# api did not return the submitted changelist number, so as a quick work around we just logged the original changelist number.
- Since the C# api was fixed and now returns the submitted changelist value we can log it.
- Changed the comments to show why we can have multiple records back.

diff --git a/Engine/Extras/P4VUtils/Commands/SubmitCommand.cs b/Engine/Extras/P4VUtils/Commands/SubmitCommand.cs
--- a/Engine/Extras/P4VUtils/Commands/SubmitCommand.cs
+++ b/Engine/Extras/P4VUtils/Commands/SubmitCommand.cs
@@ -496,9 +496,13 @@ namespace P4VUtils.Commands
 			bool successfulSubmit = submitResponses.All(x => x.Succeeded);
 			if (successfulSubmit)
 			{
-				// @todo The submit API really should return the number that the changelist was finally submit as so that
-				// we can log that here instead.
-				logger.LogInformation("Successfully submited changelist {Change}", changeNumber);
+				// The submit request will return a number of records. One for the original changelist (null)
+				// one for each submitted file (all null) and the last record detailing the submitted changelist.
+				// So we can just grab the submitted changelist number from the last record.
+
+				PerforceResponse<SubmitRecord> submittedResponse = submitResponses[submitResponses.Count - 1];
+				logger.LogInformation("Successfully submited CL {SrcCL} as CL {DstCL}", changeNumber, submittedResponse.Data.SubmittedChangeNumber);
+				
 				return true;
 			}
 			else
