Add a P4V command to open up a the preflight page in a browser if there is a #preflight tag in a changelist description.

#rnx

- We search the changelist description for '#preflight ???' and then open the appropriate page.
- If no tag is found we will display a message to the user via a modal dialog so that they are aware that no page will be opened.
- Refactored PreflightCommand.OpenUrl to be a public static method to avoid copy and pasting code.
- Refactored PreflightCommand.GetUrl to move the extraction of the Horde server address to a public static method to improve code re-use.
- The command will be added to P4 so that it can be used on pending and submitted changelists.

diff --git a/Engine/Extras/P4VUtils/Commands/PreflightCommand.cs b/Engine/Extras/P4VUtils/Commands/PreflightCommand.cs
--- a/Engine/Extras/P4VUtils/Commands/PreflightCommand.cs
+++ b/Engine/Extras/P4VUtils/Commands/PreflightCommand.cs
@@ -155,18 +155,24 @@ namespace P4VUtils.Commands
 		}
 
 		public virtual string GetUrl(string Stream, int Change, IReadOnlyDictionary<string, string> ConfigValues)
+		{
+			string BaseUrl = PreflightCommand.GetHordeServerAddress(ConfigValues);
+			return $"{BaseUrl}/preflight?stream={Stream}&change={Change}";
+		}
+
+		public virtual bool IsSubmit() { return false; }
+
+		public static string GetHordeServerAddress(IReadOnlyDictionary<string, string> ConfigValues)
 		{
 			string? BaseUrl;
 			if (!ConfigValues.TryGetValue("HordeServer", out BaseUrl))
 			{
 				BaseUrl = "https://configure-server-url-in-p4vutils.ini";
 			}
-			return $"{BaseUrl.TrimEnd('/')}/preflight?stream={Stream}&change={Change}";
-		}
 
-		public virtual bool IsSubmit() { return false; }
-
-		static void OpenUrl(string Url)
+			return BaseUrl.TrimEnd('/');
+		}
+		public static void OpenUrl(string Url)
 		{
 			if (RuntimeInformation.IsOSPlatform(OSPlatform.Windows))
 			{
@@ -208,4 +214,76 @@ namespace P4VUtils.Commands
 			return true;
 		}
 	}
+
+	class OpenPreflightCommand : Command
+	{
+		public override string Description => "If the changelist has been tagged with #preflight, open the preflight Horde page in the browser";
+
+		public override CustomToolInfo CustomTool => new CustomToolInfo("Horde: Open Preflight in browser...", "$c %c");
+
+		public override async Task<int> Execute(string[] args, IReadOnlyDictionary<string, string> configValues, ILogger logger)
+		{
+			logger.LogInformation("Parsing args ...");
+
+			// Parse command lines
+			if (args.Length < 3)
+			{
+				logger.LogError("Not enough args for command, tool is now exiting");
+				return 1;
+			}
+
+			string clientSpec = args[1];
+
+			if (!int.TryParse(args[2], out int changeNumber))
+			{
+				logger.LogError("'{Argument}' is not a numbered changelist, tool is now exiting", args[2]);
+				return 1;
+			}
+
+			logger.LogInformation("Connecting to Perforce...");
+
+			// We prefer the native client to avoid the problem where different versions of p4.exe expect
+			// or return records with different formatting to each other.
+			PerforceSettings settings = new PerforceSettings(PerforceSettings.Default) { PreferNativeClient = true, ClientName = clientSpec };
+
+			using IPerforceConnection perforceConnection = await PerforceConnection.CreateAsync(settings, logger);
+			if (perforceConnection == null)
+			{
+				logger.LogError("Failed to connect to Perforce, tool is now exiting");
+				return 1;
+			}
+
+			DescribeRecord description = await perforceConnection.DescribeAsync(changeNumber);
+
+			MatchCollection matches = Regex.Matches(description.Description, @"#preflight ([0-9a-fA-F]{24})$", RegexOptions.Multiline);
+			if (matches.Count == 0)
+			{
+				string message = $"Description for {changeNumber} does not contain any valid preflight tags";
+
+				logger.LogInformation("{Message}", message);			
+				MessageBoxResult result = MessageBox.Show(message, "No Preflights Found", MessageBoxButton.OK, MessageBoxImage.Information);
+
+				return 0;
+			}
+
+			logger.LogInformation("Found {Count} preflight tag(s)", matches.Count);
+
+			foreach (Match? match in matches)
+			{
+				// Fairly sure that match will not be null
+				string preflightURL = GetUrl(match!.Groups[1].Value, configValues);
+
+				logger.LogInformation("Opening URL '{URL}' in browser", preflightURL);
+				PreflightCommand.OpenUrl(preflightURL);
+			}
+
+			return 0;
+		}
+
+		private static string GetUrl(string preflightId, IReadOnlyDictionary<string, string> ConfigValues)
+		{
+			string BaseUrl = PreflightCommand.GetHordeServerAddress(ConfigValues);
+			return $"{BaseUrl}/job/{preflightId}";
+		}
+	}
 }
diff --git a/Engine/Extras/P4VUtils/Program.cs b/Engine/Extras/P4VUtils/Program.cs
--- a/Engine/Extras/P4VUtils/Program.cs
+++ b/Engine/Extras/P4VUtils/Program.cs
@@ -88,6 +88,7 @@ namespace P4VUtils
 			["preflight"] = new PreflightCommand(),
 			["preflightandsubmit"] = new PreflightAndSubmitCommand(),
 			["movewriteablepreflightandsubmit"] = new MoveWriteableFilesthenPreflightAndSubmitCommand(),
+			["openpreflight"] = new OpenPreflightCommand(),
 		};
 
 		public static IDictionary<string, Command> Commands = SubmissionCommands.Concat(RootHelperCommands).Concat(ToolboxCommands).Concat(IntegrateCommands).Concat(HordeCommands).ToDictionary(p => p.Key, p => p.Value, StringComparer.OrdinalIgnoreCase);
