P4VUtils: Add Messagebox for common issues wth changelists
#rnx

diff --git a/Engine/Extras/P4VUtils/Commands/PreflightCommand.cs b/Engine/Extras/P4VUtils/Commands/PreflightCommand.cs
--- a/Engine/Extras/P4VUtils/Commands/PreflightCommand.cs
+++ b/Engine/Extras/P4VUtils/Commands/PreflightCommand.cs
@@ -36,11 +36,26 @@ namespace P4VUtils.Commands
 			if (Args.Length < 2)
 			{
 				Logger.LogError("Missing changelist number");
+				MessageBoxResult result = MessageBox.Show(
+					"This option requires a changelist number to be provided\r\n" +
+					"\r\n" +
+					"Please ensure the tool is properly installed and try again \r\n",
+
+					"Invalid Tool Installation?",
+					MessageBoxButton.OK);
 				return 1;
 			}
 			else if (!int.TryParse(Args[1], out Change))
 			{
 				Logger.LogError("'{Argument}' is not a numbered changelist", Args[1]);
+				MessageBoxResult result = MessageBox.Show(
+					"This option doesn't currently support using the default changelist\r\n" +
+					"\r\n" +
+					"Please manually move the files into a non-default \r\n" +
+					"changelist on Perforce and try the operation again\r\n",
+
+					"This changelist requires manual fixes",
+					MessageBoxButton.OK);
 				return 1;
 			}
 
