Improve the parameters passed to UnrealVirtualizationTool

- Pass the project path as the first parameter. We already sort the packages by project and invoke the tool once per project so we might as well tell the tool which project to use. If we don't then UVT will need to find the project file and then relaunch itself which is a waste of time.
- Add quotes around the filepaths we pass into UnrealVirtualizationTool so that paths with spaces in them are correctly parsed.

diff --git a/Engine/Extras/P4VUtils/Commands/SubmitCommand.cs b/Engine/Extras/P4VUtils/Commands/SubmitCommand.cs
--- a/Engine/Extras/P4VUtils/Commands/SubmitCommand.cs
+++ b/Engine/Extras/P4VUtils/Commands/SubmitCommand.cs
@@ -172,7 +172,8 @@ namespace P4VUtils.Commands
 
 				foreach (KeyValuePair<string, List<string>> project in projects)
 				{
-					string engineRoot = GetEngineRootForProject(project.Key, logger);
+					string projectPath = project.Key;
+					string engineRoot = GetEngineRootForProject(projectPath, logger);
 					if (!String.IsNullOrEmpty(engineRoot))
 					{
 						logger.NewLine();
@@ -198,7 +199,7 @@ namespace P4VUtils.Commands
 						// the tool to be built.
 						logger.LogInformation("PackageList was written to '{Path}'", tempFilesPath);
 
-						if (await RunVirtualizationTool(engineRoot, clientSpec, tempFilesPath, logger) == false)
+						if (await RunVirtualizationTool(engineRoot, projectPath, clientSpec, tempFilesPath, logger) == false)
 						{
 							return false;
 						};
@@ -391,16 +392,17 @@ namespace P4VUtils.Commands
 		/// Runs the UnrealVirtualizationTool
 		/// </summary>
 		/// <param name="engineRoot">Root path of the engine we want to run the tool from</param>
+		/// <param name="projectPath">The absolute path of the project file for the packages in the package list</param>
 		/// <param name="clientSpec">The perforce client spec that the files are in</param>
 		/// <param name="packageListPath">A path to a text file containing the paths of the packages to be virtualized</param>
 		/// <param name="logger">Interface for logging</param>
 		/// <returns></returns>
-		private static async Task<bool> RunVirtualizationTool(string engineRoot, string clientSpec, string packageListPath, ILogger logger)
+		private static async Task<bool> RunVirtualizationTool(string engineRoot, string projectPath, string clientSpec, string packageListPath, ILogger logger)
 		{
 			logger.LogInformation("Running UnrealVirtualizationTool...");
 
 			string toolPath = Path.Combine(engineRoot, @"Engine\Binaries\Win64\UnrealVirtualizationTool.exe");
-			string toolArgs = string.Format("-ClientSpecName={0} -Mode=PackageList -Path={1}", clientSpec, packageListPath);
+			string toolArgs = string.Format("\"{0}\" -ClientSpecName={1} -Mode=PackageList -Path=\"{2}\"", projectPath, clientSpec, packageListPath);
 
 			using (Stream stdOutput = Console.OpenStandardOutput())
 			using (ManagedProcessGroup Group = new ManagedProcessGroup())
