[P4VUtils] Add '-MinimalLogging' commandline parameter when calling UnrealVirtualizationTool tool to maintain the old behaviour when logging was supressed by default.

diff --git a/Engine/Extras/P4VUtils/Commands/SubmitCommand.cs b/Engine/Extras/P4VUtils/Commands/SubmitCommand.cs
--- a/Engine/Extras/P4VUtils/Commands/SubmitCommand.cs
+++ b/Engine/Extras/P4VUtils/Commands/SubmitCommand.cs
@@ -402,7 +402,7 @@ namespace P4VUtils.Commands
 			logger.LogInformation("Running UnrealVirtualizationTool...");
 
 			string toolPath = Path.Combine(engineRoot, @"Engine\Binaries\Win64\UnrealVirtualizationTool.exe");
-			string toolArgs = string.Format("\"{0}\" -ClientSpecName={1} -Mode=PackageList -Path=\"{2}\"", projectPath, clientSpec, packageListPath);
+			string toolArgs = string.Format("\"{0}\" -MinimalLogging -ClientSpecName={1} -Mode=PackageList -Path=\"{2}\"", projectPath, clientSpec, packageListPath);
 
 			using (Stream stdOutput = Console.OpenStandardOutput())
 			using (ManagedProcessGroup Group = new ManagedProcessGroup())
