[P4VUtils] Upgrade to NET 6.

- When built via build graph job 'Build P4VUtils' we now produce a self contained version of the tool which will allow it to work on machines that have no yet installed .NET 6.0

diff --git a/Engine/Extras/P4VUtils/P4VUtils.csproj b/Engine/Extras/P4VUtils/P4VUtils.csproj
--- a/Engine/Extras/P4VUtils/P4VUtils.csproj
+++ b/Engine/Extras/P4VUtils/P4VUtils.csproj
@@ -1,8 +1,8 @@
-<Project Sdk="Microsoft.NET.Sdk.WindowsDesktop">
+<Project Sdk="Microsoft.NET.Sdk">
 
   <PropertyGroup>
     <OutputType>Exe</OutputType>
-    <TargetFramework>netcoreapp3.1</TargetFramework>
+    <TargetFramework>net6.0-windows</TargetFramework>
     <Nullable>enable</Nullable>
     <UseWPF>true</UseWPF>
     <EnableNETAnalyzers>True</EnableNETAnalyzers>
@@ -15,7 +15,7 @@
   </ItemGroup>
 
   <ItemGroup>
-    <PackageReference Include="Microsoft.CodeAnalysis.NetAnalyzers" Version="6.0.0">
+    <PackageReference Include="Microsoft.CodeAnalysis.NetAnalyzers" Version="7.0.0">
       <PrivateAssets>all</PrivateAssets>
       <IncludeAssets>runtime; build; native; contentfiles; analyzers; buildtransitive</IncludeAssets>
     </PackageReference>
@@ -23,16 +23,14 @@
   </ItemGroup>
 
   <ItemGroup>
-    <ProjectReference Include="..\..\Source\Programs\UnrealGameSync\Shared\EpicGames.Core\EpicGames.Core.csproj" />
-    <ProjectReference Include="..\..\Source\Programs\UnrealGameSync\Shared\EpicGames.Perforce\EpicGames.Perforce.csproj" />
+    <ProjectReference Include="..\..\Source\Programs\Shared\EpicGames.Core\EpicGames.Core.csproj" />
+    <ProjectReference Include="..\..\Source\Programs\Shared\EpicGames.Perforce\EpicGames.Perforce.csproj" />
   </ItemGroup>
-
+  
   <ItemGroup>
     <None Update="P4VUtils.ini">
       <CopyToOutputDirectory>PreserveNewest</CopyToOutputDirectory>
     </None>
   </ItemGroup>
 
-  
-
 </Project>
