[P4VUtils] Optimize backout command when run on binary files

#rnx

- Previously we were issuing a 'p4 fstat' command per binary file found, now we gather the binary files and issue a single 'p4 fstat' command for all of them.
- In testing (backing out a changelist with 10 binary files) this reduced the time taken to backout 15->10s when the server was not under strain and 30s->15s when the server was in peak usage.

diff --git a/Engine/Extras/P4VUtils/Commands/BackoutCommand.cs b/Engine/Extras/P4VUtils/Commands/BackoutCommand.cs
--- a/Engine/Extras/P4VUtils/Commands/BackoutCommand.cs
+++ b/Engine/Extras/P4VUtils/Commands/BackoutCommand.cs
@@ -45,6 +45,8 @@ namespace P4VUtils.Commands
 				Result.Errors.Add(string.Format($"CL {Change} did not originate in the current stream, to be safe, CLs should be backed out where the change was originally submitted.", Change));
 			}
 
+			List<DescribeFileRecord> BinaryFiles = new List<DescribeFileRecord>();
+
 			foreach (DescribeFileRecord DescFileRec in DescribeRec.Files)
 			{
 				// is this an edit/add/delete operation?
@@ -63,18 +65,38 @@ namespace P4VUtils.Commands
 					Result.Errors.Add($"In CL {Change}, the file below affects asset versioning, therefore it isn't safe to back it out.\n\n{DescFileRec.DepotFile}");
 				}
 
-				// if it's a binary file, are we backing out the head revision?
+				// binary files need another round of validation later
 				if (DescFileRec.Type.Contains("+l", StringComparison.InvariantCultureIgnoreCase))
 				{
-					// check the latest revision of this depot file
-					FStatRecord StatRecord = (await Perforce.FStatAsync(new List<string> { DescFileRec.DepotFile }).ToListAsync()).First();
-					if (StatRecord.HeadRevision != DescFileRec.Revision)
+					BinaryFiles.Add(DescFileRec);
+				}
+			}
+
+			if (BinaryFiles.Count > 0)
+			{
+				// we can only safely backout binary files that are at the head revision, so check for that/
+
+				FileSpecList FileSpecs = BinaryFiles.Select(x => x.DepotFile).ToList();
+
+				List<FStatRecord> StatRecord = await Perforce.FStatAsync(FileSpecs).ToListAsync();
+
+				if(BinaryFiles.Count == StatRecord.Count)
+				{
+					for( int Index = 0; Index < BinaryFiles.Count; ++Index)
 					{
-						Logger.LogError("\nIn CL {Change}, the file below is a binary file and revision #{Rev} isn't the head revision (#{HeadRev}),\ntherefore it isn't safe to back it out\n\n{FileName}",
-							Change, DescFileRec.Revision, StatRecord.HeadRevision, DescFileRec.DepotFile);
-						Result.Errors.Add($"\nIn CL {Change}, the file below is a binary file and revision #{DescFileRec.Revision} isn't the head revision (#{StatRecord.HeadRevision}), therefore it isn't safe to back it out\n\n{DescFileRec.DepotFile}");
+						if (StatRecord[Index].HeadRevision != BinaryFiles[Index].Revision)
+						{
+							Logger.LogError("\nIn CL {Change}, the file below is a binary file and revision #{Rev} isn't the head revision (#{HeadRev}),\ntherefore it isn't safe to back it out\n\n{FileName}",
+								Change, BinaryFiles[Index].Revision, StatRecord[Index].HeadRevision, BinaryFiles[Index].DepotFile);
+								Result.Errors.Add($"\nIn CL {Change}, the file below is a binary file and revision #{BinaryFiles[Index].Revision} isn't the head revision (#{StatRecord[Index].HeadRevision}), therefore it isn't safe to back it out\n\n{BinaryFiles[Index].DepotFile}");
+						}
 					}
 				}
+				else
+				{
+					// This shouldn't happen but if fstat did not return the expected number of records then we cannot safely continue
+					throw new FatalErrorException("Error running ValidateBackoutSafety\nThe number of records returned from p4 fstat {0} does not match the number of records requested {1}'", StatRecord.Count, BinaryFiles.Count);
+				}
 			}
 
 			if (Result.Errors.Count == 0)
