P4VUtils: Refactoring command discovering to be dynamic

Allows adding internal commands for Epic use only

#rnx

diff --git a/Engine/Extras/P4VUtils/CommandAttribute.cs b/Engine/Extras/P4VUtils/CommandAttribute.cs
new file mode 100644
--- /dev/null
+++ b/Engine/Extras/P4VUtils/CommandAttribute.cs
@@ -0,0 +1,67 @@
+// Copyright Epic Games, Inc. All Rights Reserved.
+
+using System;
+
+namespace P4VUtils
+{
+	public enum CommandCategory
+	{
+		/// <summary>
+		/// Commands that help with common but simple operations
+		/// </summary>
+		Root,
+
+		/// <summary>
+		/// commands that help with submitting files/changelists
+		/// </summary>
+		Submission,
+
+		/// <summary>
+		/// Commands that help with common but simple operations
+		/// </summary>
+		Toolbox,
+
+		/// <summary>
+		/// complex commands to facilitate integrations/backout
+		/// </summary>
+		Integrate,
+
+		/// <summary>
+		/// local build and horde preflights
+		/// </summary>
+		Horde
+	}
+
+	/// <summary>
+	/// Attribute used to define a class as a Command to hook into P4V
+	/// 
+	/// Containing attributes to define the cli name of the command, the categorization of the command, and the order it appears in the UI.
+	/// </summary>
+	[AttributeUsage(AttributeTargets.Class)]
+	public sealed class CommandAttribute : System.Attribute
+	{
+		/// <summary>
+		/// Command name, use to reference the command via command line args
+		/// </summary>
+		public string CommandName { get; internal set; } = null!;
+
+		/// <summary>
+		/// Category to organize into UI Elements. 
+		/// </summary>
+		public CommandCategory Category { get; internal set; }
+
+		/// <summary>
+		/// Used to preserve previous order.
+		/// 
+		/// Defaults to int.MaxValue to ensure any new items appear at the end of the sub menus.
+		/// </summary>
+		public int Order { get; internal set; }
+
+		public CommandAttribute(string commandName, CommandCategory category, int order = int.MaxValue)
+		{
+			CommandName = commandName;
+			Category = category;
+			Order = order;
+		}
+	}
+}
diff --git a/Engine/Extras/P4VUtils/Commands/BackoutCommand.cs b/Engine/Extras/P4VUtils/Commands/BackoutCommand.cs
--- a/Engine/Extras/P4VUtils/Commands/BackoutCommand.cs
+++ b/Engine/Extras/P4VUtils/Commands/BackoutCommand.cs
@@ -20,6 +20,7 @@ namespace P4VUtils.Commands
 		public List<String> Errors = new List<String>();
 	}
 
+	[Command("backout", CommandCategory.Toolbox, 0)]
 	class BackoutCommand : Command
 	{
 		public override string Description => "P4 Admin sanctioned method of backing out a CL";
diff --git a/Engine/Extras/P4VUtils/Commands/CherryPickCommand.cs b/Engine/Extras/P4VUtils/Commands/CherryPickCommand.cs
--- a/Engine/Extras/P4VUtils/Commands/CherryPickCommand.cs
+++ b/Engine/Extras/P4VUtils/Commands/CherryPickCommand.cs
@@ -12,6 +12,7 @@ using System.Threading.Tasks;
 
 namespace P4VUtils.Commands
 {
+	[Command("cherypick", CommandCategory.Integrate, 0)]
 	class CherryPickCommand : Command
 	{
 		public override string Description => "Cherry-picks a changelist into the current stream";
diff --git a/Engine/Extras/P4VUtils/Commands/CompileCommand.cs b/Engine/Extras/P4VUtils/Commands/CompileCommand.cs
--- a/Engine/Extras/P4VUtils/Commands/CompileCommand.cs
+++ b/Engine/Extras/P4VUtils/Commands/CompileCommand.cs
@@ -14,6 +14,7 @@ using System.Threading.Tasks;
 
 namespace P4VUtils.Commands
 {
+	[Command("compile", CommandCategory.Horde, 0)]
 	class CompileCommand : Command
 	{
 		public override string Description => "Compiles files in the selected changelist";
diff --git a/Engine/Extras/P4VUtils/Commands/ConvertToEditCommand.cs b/Engine/Extras/P4VUtils/Commands/ConvertToEditCommand.cs
--- a/Engine/Extras/P4VUtils/Commands/ConvertToEditCommand.cs
+++ b/Engine/Extras/P4VUtils/Commands/ConvertToEditCommand.cs
@@ -13,6 +13,7 @@ using System.Threading.Tasks;
 
 namespace P4VUtils.Commands
 {
+	[Command("converttoedit", CommandCategory.Integrate, 1)]
 	class ConvertToEditCommand : Command
 	{
 		public override string Description => "Converts an integration to an edit";
diff --git a/Engine/Extras/P4VUtils/Commands/CopyCLCommand.cs b/Engine/Extras/P4VUtils/Commands/CopyCLCommand.cs
--- a/Engine/Extras/P4VUtils/Commands/CopyCLCommand.cs
+++ b/Engine/Extras/P4VUtils/Commands/CopyCLCommand.cs
@@ -13,6 +13,7 @@ using System.Windows;
 
 namespace P4VUtils.Commands
 {
+	[Command("copyclnum", CommandCategory.Root, 1)]
 	class CopyCLCommand : Command
 	{
 		public override string Description => "copy CL # only from pending and submitted CLs";
diff --git a/Engine/Extras/P4VUtils/Commands/DataCommands.cs b/Engine/Extras/P4VUtils/Commands/DataCommands.cs
--- a/Engine/Extras/P4VUtils/Commands/DataCommands.cs
+++ b/Engine/Extras/P4VUtils/Commands/DataCommands.cs
@@ -29,11 +29,12 @@ namespace P4VUtils.Commands
 		}
 	}
 
+	[Command("convertcldatatolocalwritable", CommandCategory.Toolbox, 6)]
 	class ConvertCLDataToLocalWritable : DataCommandBase
 	{
 		public override string Description => "Convert all data files in changelist to local writable";
 
-		public override CustomToolInfo CustomTool => new CustomToolInfo("Convert data to local writable", "%c");
+		public override CustomToolInfo CustomTool => new CustomToolInfo("Convert changelist to local writable", "%c");
 
 		public override async Task<int> Execute(string[] Args, IReadOnlyDictionary<string, string> ConfigValues, ILogger Logger)
 		{
@@ -57,11 +58,12 @@ namespace P4VUtils.Commands
 		}
 	}
 
+	[Command("convertdatatolocalwritable", CommandCategory.Toolbox, 7)]
 	class ConvertDataToLocalWritable : DataCommandBase
 	{
 		public override string Description => "Convert all selected data files to local writable";
 
-		public override CustomToolInfo CustomTool => new CustomToolInfo("Convert data to local writable", "%F");
+		public override CustomToolInfo CustomTool => new CustomToolInfo("Convert selected files to local writable", "%F");
 
 		public override async Task<int> Execute(string[] Args, IReadOnlyDictionary<string, string> ConfigValues, ILogger Logger)
 		{
diff --git a/Engine/Extras/P4VUtils/Commands/DescribeCommand.cs b/Engine/Extras/P4VUtils/Commands/DescribeCommand.cs
--- a/Engine/Extras/P4VUtils/Commands/DescribeCommand.cs
+++ b/Engine/Extras/P4VUtils/Commands/DescribeCommand.cs
@@ -15,6 +15,7 @@ using System.Diagnostics;
 
 namespace P4VUtils.Commands
 {
+	[Command("describe", CommandCategory.Root, 0)]
 	class DescribeCommand : Command
 	{
 		public override string Description => "Prints out a description of the given changelist";
@@ -47,6 +48,7 @@ namespace P4VUtils.Commands
 		}
 	}
 
+	[Command("describedirectory", CommandCategory.Toolbox, 8)]
 	class DescribeDirectoryCommand : Command
 	{
 		public override string Description => "Prints out a description of last 30 days of CL's in the given directory to txt and csv";
diff --git a/Engine/Extras/P4VUtils/Commands/EdigrateCommand.cs b/Engine/Extras/P4VUtils/Commands/EdigrateCommand.cs
--- a/Engine/Extras/P4VUtils/Commands/EdigrateCommand.cs
+++ b/Engine/Extras/P4VUtils/Commands/EdigrateCommand.cs
@@ -12,7 +12,7 @@ using System.Threading.Tasks;
 
 namespace P4VUtils.Commands
 {
-
+	[Command("edigrate", CommandCategory.Integrate, 2)]
 	class EdigrateCommand : Command
 	{
 		public override string Description => "Cherry-picks a change to the current stream as an edit";
diff --git a/Engine/Extras/P4VUtils/Commands/FastReconcileEdits.cs b/Engine/Extras/P4VUtils/Commands/FastReconcileEdits.cs
--- a/Engine/Extras/P4VUtils/Commands/FastReconcileEdits.cs
+++ b/Engine/Extras/P4VUtils/Commands/FastReconcileEdits.cs
@@ -17,6 +17,7 @@ using System.Threading.Tasks;
 
 namespace P4VUtils.Commands
 {
+	[Command("reconcilecode", CommandCategory.Toolbox, 2)]
 	class FastReconcileCodeEditsCommand : Command
 	{
 		
@@ -105,6 +106,7 @@ namespace P4VUtils.Commands
 		}
 	}
 
+	[Command("reconcileall", CommandCategory.Toolbox, 3)]
 	class FastReconcileAllEditsCommand : FastReconcileCodeEditsCommand
 	{
 		public override string[] GetExtensions()
diff --git a/Engine/Extras/P4VUtils/Commands/FindLastEditCommand.cs b/Engine/Extras/P4VUtils/Commands/FindLastEditCommand.cs
--- a/Engine/Extras/P4VUtils/Commands/FindLastEditCommand.cs
+++ b/Engine/Extras/P4VUtils/Commands/FindLastEditCommand.cs
@@ -11,6 +11,7 @@ using P4VUtils.Perforce;
 
 namespace P4VUtils.Commands
 {
+	[Command("findlastedit", CommandCategory.Root, 2)]
 	class FindLastEditCommand : Command
 	{
 		public override string Description => "Prints who last edited this file and in which stream";
diff --git a/Engine/Extras/P4VUtils/Commands/P4BlameCommand.cs b/Engine/Extras/P4VUtils/Commands/P4BlameCommand.cs
--- a/Engine/Extras/P4VUtils/Commands/P4BlameCommand.cs
+++ b/Engine/Extras/P4VUtils/Commands/P4BlameCommand.cs
@@ -9,6 +9,7 @@ using System.Threading.Tasks;
 
 namespace P4VUtils.Commands
 {
+	[Command("findlasteditbyline", CommandCategory.Root, 3)]
 	class P4BlameCommand : Command
 	{
 		public override string Description => "Prints out who last edited a specific line in a file";
diff --git a/Engine/Extras/P4VUtils/Commands/PreflightCommand.cs b/Engine/Extras/P4VUtils/Commands/PreflightCommand.cs
--- a/Engine/Extras/P4VUtils/Commands/PreflightCommand.cs
+++ b/Engine/Extras/P4VUtils/Commands/PreflightCommand.cs
@@ -19,6 +19,7 @@ using System.Linq;
 
 namespace P4VUtils.Commands
 {
+	[Command("preflight", CommandCategory.Horde, 1)]
 	class PreflightCommand : Command
 	{
 
@@ -204,6 +205,7 @@ namespace P4VUtils.Commands
 		}
 	}
 
+	[Command("preflightandsubmit", CommandCategory.Horde, 2)]
 	class PreflightAndSubmitCommand : PreflightCommand
 	{
 		public override string Description => "Runs a preflight of the given changelist on Horde and submits it";
@@ -218,6 +220,7 @@ namespace P4VUtils.Commands
 		public override bool IsSubmit() { return true; }
 	}
 
+	[Command("movewriteablepreflightandsubmit", CommandCategory.Horde, 3)]
 	class MoveWriteableFilesthenPreflightAndSubmitCommand : PreflightAndSubmitCommand
 	{
 		public override string Description => "Moves the writeable files to a new CL, then runs a preflight of the current changelist on Horde and submits it";
@@ -230,6 +233,7 @@ namespace P4VUtils.Commands
 		}
 	}
 
+	[Command("openpreflight", CommandCategory.Horde, 4)]
 	class OpenPreflightCommand : Command
 	{
 		public override string Description => "If the changelist has been tagged with #preflight, open the preflight Horde page in the browser";
diff --git a/Engine/Extras/P4VUtils/Commands/SnapShotCommand.cs b/Engine/Extras/P4VUtils/Commands/SnapShotCommand.cs
--- a/Engine/Extras/P4VUtils/Commands/SnapShotCommand.cs
+++ b/Engine/Extras/P4VUtils/Commands/SnapShotCommand.cs
@@ -18,6 +18,7 @@ using System.Threading.Tasks;
 
 namespace P4VUtils.Commands
 {
+	[Command("snapshot", CommandCategory.Toolbox, 1)]
 	class SnapshotCommand : Command
 	{
 		public override string Description => "Creates a shelved copy of the CL";
diff --git a/Engine/Extras/P4VUtils/Commands/SubmitCommand.cs b/Engine/Extras/P4VUtils/Commands/SubmitCommand.cs
--- a/Engine/Extras/P4VUtils/Commands/SubmitCommand.cs
+++ b/Engine/Extras/P4VUtils/Commands/SubmitCommand.cs
@@ -58,6 +58,7 @@ namespace P4VUtils.Commands
 		}
 	}
 
+	[Command("submitandvirtualize", CommandCategory.Submission, 0)]
 	class SubmitAndVirtualizeCommand : Command
 	{
 		public override string Description => "Virtualize And Submit";
diff --git a/Engine/Extras/P4VUtils/Commands/UnshelveCommands.cs b/Engine/Extras/P4VUtils/Commands/UnshelveCommands.cs
--- a/Engine/Extras/P4VUtils/Commands/UnshelveCommands.cs
+++ b/Engine/Extras/P4VUtils/Commands/UnshelveCommands.cs
@@ -125,6 +125,7 @@ namespace P4VUtils.Commands
 		}
 	}
 
+	[Command("unshelvetocurrentrevision", CommandCategory.Toolbox, 4)]
 	class UnshelveToCurrentRevision : UnshelveCommandBase
 	{
 		public override string Description => "Remembers revision all files you are about to unshelve and if the revision is older, will sync to saved revision";
@@ -149,6 +150,7 @@ namespace P4VUtils.Commands
 		}
 	}
 
+	[Command("unshelvemakedatawritable", CommandCategory.Toolbox, 5)]
 	class UnshelveMakeDataWritable : UnshelveCommandBase
 	{
 		public override string Description => "Unshelve changelist but for data will make a writable copy locally";
diff --git a/Engine/Extras/P4VUtils/P4VUtils.csproj b/Engine/Extras/P4VUtils/P4VUtils.csproj
--- a/Engine/Extras/P4VUtils/P4VUtils.csproj
+++ b/Engine/Extras/P4VUtils/P4VUtils.csproj
@@ -6,6 +6,10 @@
     <Nullable>enable</Nullable>
     <UseWPF>true</UseWPF>
     <EnableNETAnalyzers>True</EnableNETAnalyzers>
+    <!-- Supress CA1812: "Class" is an internal class that is apparently never instantiated.
+      With the change to dynamic instantiation this warning is triggered.
+      -->
+    <NoWarn>CA1812</NoWarn>
   </PropertyGroup>
 
   <ItemGroup Condition="Exists('..\..\Restricted\NotForLicensees\Extras\P4VUtils\P4VUtils.ini')">
@@ -27,6 +31,12 @@
     <ProjectReference Include="..\..\Source\Programs\Shared\EpicGames.Perforce\EpicGames.Perforce.csproj" />
   </ItemGroup>
   
+  <ItemGroup>
+    <Compile Include="..\..\Restricted\*\Extras\P4VUtils\Commands\*.cs">
+      <Link>Restricted\Commands\%(Filename)%(Extension)</Link>
+    </Compile>
+  </ItemGroup>
+
   <ItemGroup>
     <None Update="P4VUtils.ini">
       <CopyToOutputDirectory>PreserveNewest</CopyToOutputDirectory>
diff --git a/Engine/Extras/P4VUtils/P4VUtils.ini b/Engine/Extras/P4VUtils/P4VUtils.ini
--- a/Engine/Extras/P4VUtils/P4VUtils.ini
+++ b/Engine/Extras/P4VUtils/P4VUtils.ini
@@ -1 +1,2 @@
 HordeServer=https://my-horde-server/
+RobomergePreviewURL=https://my-robomerge-server/preview/
\ No newline at end of file
diff --git a/Engine/Extras/P4VUtils/Program.cs b/Engine/Extras/P4VUtils/Program.cs
--- a/Engine/Extras/P4VUtils/Program.cs
+++ b/Engine/Extras/P4VUtils/Program.cs
@@ -2,11 +2,8 @@
 
 using EpicGames.Core;
 using Microsoft.Extensions.Logging;
-using Microsoft.Extensions.Logging.Console;
-using P4VUtils.Commands;
 using System;
 using System.Collections.Generic;
-using System.ComponentModel;
 using System.IO;
 using System.Linq;
 using System.Reflection;
@@ -44,54 +41,8 @@ namespace P4VUtils
 
 	class Program
 	{
-		// UEHelpersInRoot - commands that help with common but simple operations
-		public static IReadOnlyDictionary<string, Command> RootHelperCommands { get; } = new Dictionary<string, Command>(StringComparer.OrdinalIgnoreCase)
-		{
-			["describe"] = new DescribeCommand(),
-			["copyclnum"] = new CopyCLCommand(),
-			["findlastedit"] = new FindLastEditCommand(),
-			["findlasteditbyline"] = new P4BlameCommand(),
-		};
-
-		// UESubmit - commands that help with submitting files/changelists
-		public static IReadOnlyDictionary<string, Command> SubmissionCommands { get; } = new Dictionary<string, Command>(StringComparer.OrdinalIgnoreCase)
-		{
-			["submitandvirtualize"] = new SubmitAndVirtualizeCommand(),
-		};
-
-		// UEHelpers - commands that help with common but simple operations
-		public static IReadOnlyDictionary<string, Command> ToolboxCommands { get; } = new Dictionary<string, Command>(StringComparer.OrdinalIgnoreCase)
-		{
-			["backout"] = new BackoutCommand(),
-			["snapshot"] = new SnapshotCommand(),
-			["reconcilecode"] = new FastReconcileCodeEditsCommand(),
-			["reconcileall"] = new FastReconcileAllEditsCommand(),
-			["unshelvetocurrentrevision"] = new UnshelveToCurrentRevision(),
-			["unshelvemakedatawritable"] = new UnshelveMakeDataWritable(),
-			["convertcldatatolocalwritable"] = new ConvertCLDataToLocalWritable(),
-			["convertdatatolocalwritable"] = new ConvertDataToLocalWritable(),
-			["describedirectory"] = new DescribeDirectoryCommand(),
-		};
-
-		// UEIntegrate Folder commands - complex commands to facilitate integrations/backout
-		public static IReadOnlyDictionary<string, Command> IntegrateCommands { get; } = new Dictionary<string, Command>(StringComparer.OrdinalIgnoreCase)
-		{
-			["cherrypick"] = new CherryPickCommand(),
-			["converttoedit"] = new ConvertToEditCommand(),
-			["edigrate"] = new EdigrateCommand(),
-		};
-
-		// UEHorde Folder - local build and horde preflights
-		public static IReadOnlyDictionary<string, Command> HordeCommands { get; } = new Dictionary<string, Command>(StringComparer.OrdinalIgnoreCase)
-		{
-			["compile"] = new CompileCommand(),
-			["preflight"] = new PreflightCommand(),
-			["preflightandsubmit"] = new PreflightAndSubmitCommand(),
-			["movewriteablepreflightandsubmit"] = new MoveWriteableFilesthenPreflightAndSubmitCommand(),
-			["openpreflight"] = new OpenPreflightCommand(),
-		};
-
-		public static IDictionary<string, Command> Commands = SubmissionCommands.Concat(RootHelperCommands).Concat(ToolboxCommands).Concat(IntegrateCommands).Concat(HordeCommands).ToDictionary(p => p.Key, p => p.Value, StringComparer.OrdinalIgnoreCase);
+		public static IDictionary<string, Command> Commands = new Dictionary<string, Command>();
+		public static IDictionary<CommandCategory, List<string>> CommandCategories = new Dictionary<CommandCategory, List<string>>();
 
 		static void PrintHelp(ILogger Logger)
 		{
@@ -114,7 +65,7 @@ namespace P4VUtils
 
 		static async Task<int> Main(string[] Args)
 		{
-			using ILoggerFactory Factory = LoggerFactory.Create(Builder => Builder.AddEpicDefault());//.AddSimpleConsole(Options => { Options.SingleLine = true; Options.IncludeScopes = false; }));
+			using ILoggerFactory Factory = LoggerFactory.Create(Builder => Builder.AddEpicDefault());
 			ILogger Logger = Factory.CreateLogger<Program>();
 			Log.SetInnerLogger(Logger);
 
@@ -129,8 +80,33 @@ namespace P4VUtils
 			}
 		}
 
+		static void RegisterCommands()
+		{
+			// Grab all the commands, sorting to preserve the previous order from the explicit tabs.
+			var Types = from assembly in Assembly.GetExecutingAssembly().GetTypes()
+						where assembly.IsDefined(typeof(CommandAttribute))
+						let attrib = assembly.GetCustomAttributes(typeof(CommandAttribute), true).Cast<CommandAttribute>().First()
+						orderby attrib.Order
+						select new { Type = assembly, Attribute = attrib };
+						
+			foreach (var Type in Types)
+			{
+				CommandAttribute attribute = Type.Attribute;
+				if (!CommandCategories.ContainsKey(attribute.Category))
+				{
+					CommandCategories[attribute.Category] = new List<string>();
+				}
+
+				Command? CreatedCommand = Activator.CreateInstance(Type.Type) as Command;
+				Commands.Add(attribute.CommandName, CreatedCommand!);
+				CommandCategories[attribute.Category].Add(attribute.CommandName);
+			}
+		}
+
 		static async Task<int> InnerMain(string[] Args, ILogger Logger)
 		{
+			RegisterCommands();
+
 			if (Args.Length == 0 || Args[0].Equals("-help", StringComparison.OrdinalIgnoreCase))
 			{
 				PrintHelp(Logger);
@@ -168,7 +144,7 @@ namespace P4VUtils
 				return await Command.Execute(Args, ConfigValues, Logger);
 			}
 			else
-		{
+			{
 				Logger.LogError("Unknown command: {Command}", Args[0]);
 				PrintHelp(Logger);
 				return 1;
@@ -273,7 +249,7 @@ namespace P4VUtils
 			return ToolsChecked == ToolsRemoved;
 		}
 
-		static void InstallCommandsListInFolder(string FolderName, bool AddFolderToContextMenu, IReadOnlyDictionary<string, Command> InputCommmands, XmlDocument Document, FileReference DotNetLocation, FileReference AssemblyLocation, ILogger Logger)
+		static void InstallCommandsListInFolder(string FolderName, bool AddFolderToContextMenu, CommandCategory Category, XmlDocument Document, FileReference DotNetLocation, FileReference AssemblyLocation, ILogger Logger)
 		{
 			// <CustomToolDefList>				// list of custom tools (top level)
 			//  < CustomToolDef >				// loose custom tool in top level
@@ -299,7 +275,10 @@ namespace P4VUtils
 
 				XmlElement FolderDefList = Document.CreateElement("CustomToolDefList");
 
-				foreach (KeyValuePair<string, Command> Pair in InputCommmands)
+				List<string> CommandList = CommandCategories[Category];
+				IEnumerable<KeyValuePair<string, Command>> CategoryCommands = Commands.Where(Pair => CommandList.Contains(Pair.Key));
+
+				foreach (KeyValuePair<string, Command> Pair in CategoryCommands)
 				{
 					CustomToolInfo CustomTool = Pair.Value.CustomTool;
 
@@ -448,11 +427,11 @@ namespace P4VUtils
 			// Insert new entries
 			if (bInstall)
 			{
-				InstallCommandsListInFolder("UE RootHelpers", false/*AddFolderToContextMenu*/, RootHelperCommands, Document, DotNetLocation, AssemblyLocation, Logger);
-				InstallCommandsListInFolder("UE Submit", true/*AddFolderToContextMenu*/, SubmissionCommands, Document, DotNetLocation, AssemblyLocation, Logger);
-				InstallCommandsListInFolder("UE Toolbox", true/*AddFolderToContextMenu*/, ToolboxCommands, Document, DotNetLocation, AssemblyLocation, Logger);
-				InstallCommandsListInFolder("UE Integrate", true/*AddFolderToContextMenu*/, IntegrateCommands, Document, DotNetLocation, AssemblyLocation, Logger);
-				InstallCommandsListInFolder("UE Horde", true/*AddFolderToContextMenu*/, HordeCommands, Document, DotNetLocation, AssemblyLocation, Logger);
+				InstallCommandsListInFolder("UE RootHelpers", false/*AddFolderToContextMenu*/, CommandCategory.Root, Document, DotNetLocation, AssemblyLocation, Logger);
+				InstallCommandsListInFolder("UE Submit", true/*AddFolderToContextMenu*/, CommandCategory.Submission, Document, DotNetLocation, AssemblyLocation, Logger);
+				InstallCommandsListInFolder("UE Toolbox", true/*AddFolderToContextMenu*/, CommandCategory.Toolbox, Document, DotNetLocation, AssemblyLocation, Logger);
+				InstallCommandsListInFolder("UE Integrate", true/*AddFolderToContextMenu*/, CommandCategory.Integrate, Document, DotNetLocation, AssemblyLocation, Logger);
+				InstallCommandsListInFolder("UE Horde", true/*AddFolderToContextMenu*/, CommandCategory.Horde, Document, DotNetLocation, AssemblyLocation, Logger);
 			}
 
 			// Save the new document
