P4VUtils: New Command to Preview RoboMerge Graph Changes

New robomergepreviewgraph - Shelves a pending CL and opens the robomerge.url/preview/<CL> page to preview graph changes.

#rnx

diff --git a/Engine/Extras/P4VUtils/Commands/RobomergePreviewCommand.cs b/Engine/Extras/P4VUtils/Commands/RobomergePreviewCommand.cs
new file mode 100644
--- /dev/null
+++ b/Engine/Extras/P4VUtils/Commands/RobomergePreviewCommand.cs
@@ -0,0 +1,104 @@
+// Copyright Epic Games, Inc. All Rights Reserved.
+
+using EpicGames.Perforce;
+using Microsoft.Extensions.Logging;
+using System;
+using System.Collections.Generic;
+using System.Diagnostics;
+using System.Linq;
+using System.Runtime.InteropServices;
+using System.Threading;
+using System.Threading.Tasks;
+
+namespace P4VUtils.Commands
+{
+	[Command("robomergepreviewgraph", CommandCategory.Toolbox, 9)]
+	class RobomergePreviewCommand : Command
+	{
+		public override string Description => "Shelves and Previews the pending changes to the Robomerge Graph.";
+
+		public override CustomToolInfo CustomTool => new CustomToolInfo("Robomerge Preview Graph", "%c") { ShowConsole = true, RefreshUI = false };
+
+		public override async Task<int> Execute(string[] Args, IReadOnlyDictionary<string, string> ConfigValues, ILogger Logger)
+		{
+			int Change;
+			if (Args.Length < 2)
+			{
+				Logger.LogError("Missing changelist number");
+				return 1;
+			}
+			else if (!int.TryParse(Args[1], out Change))
+			{
+				Logger.LogError("'{Argument}' is not a numbered changelist", Args[1]);
+				return 1;
+			}
+
+			string? ClientName = Environment.GetEnvironmentVariable("P4CLIENT");
+			using PerforceConnection Perforce = new PerforceConnection(null, null, ClientName, Logger);
+
+			ClientRecord Client = await Perforce.GetClientAsync(ClientName, CancellationToken.None);
+			if (Client.Stream == null)
+			{
+				Logger.LogError("Not being run from a stream client");
+				return 1;
+			}
+
+			List<OpenedRecord> OpenedRecords = await Perforce.OpenedAsync(OpenedOptions.None, Change, null, null, -1, FileSpecList.Any, CancellationToken.None).ToListAsync();
+
+			if (!OpenedRecords.Any(x => x.ClientFile.EndsWith(".branchmap.json", StringComparison.OrdinalIgnoreCase)))
+			{
+				Logger.LogError("No branchmapping.json files contained in changelist {Change}", Change);
+				return 1;
+			}
+
+			if (OpenedRecords.Count == 0)
+			{
+				Logger.LogError("No files checked out in changelist {Change}", Change);
+			}
+
+			Logger.LogInformation("Shelving changelist {Change}", Change);
+			await Perforce.ShelveAsync(Change, ShelveOptions.Overwrite, new[] { "//..." }, CancellationToken.None);
+
+			List<DescribeRecord> Describe = await Perforce.DescribeAsync(DescribeOptions.Shelved, -1, new[] { Change }, CancellationToken.None);
+			if (Describe[0].Files.Count == 0)
+			{
+				Logger.LogError("No files were shelved in {Change}", Change);
+				return 1;
+			}
+
+			string Url = GetUrl(Change, ConfigValues);
+			Logger.LogInformation("Opening {Url}", Url);
+			OpenUrl(Url);
+
+			return 0;
+		}
+
+		public virtual string GetUrl(int Change, IReadOnlyDictionary<string, string> ConfigValues)
+		{
+			string? BaseUrl;
+			if (!ConfigValues.TryGetValue("RobomergePreviewURL", out BaseUrl))
+			{
+				BaseUrl = "https://configure-server-url-in-p4vutils.ini";
+			}
+
+			return $"{BaseUrl.TrimEnd('/')}/{Change}";
+
+		}
+
+		static void OpenUrl(string Url)
+		{
+			if (RuntimeInformation.IsOSPlatform(OSPlatform.Windows))
+			{
+				Process.Start(new ProcessStartInfo(Url) { UseShellExecute = true });
+			}
+			else if (RuntimeInformation.IsOSPlatform(OSPlatform.Linux))
+			{
+				Process.Start("xdg-open", Url);
+			}
+			else if (RuntimeInformation.IsOSPlatform(OSPlatform.OSX))
+			{
+				Process.Start("open", Url);
+			}
+		}
+	}
+}
