- Added missing copyright

diff --git a/Engine/Extras/P4VUtils/UserInterface.cs b/Engine/Extras/P4VUtils/UserInterface.cs
--- a/Engine/Extras/P4VUtils/UserInterface.cs
+++ b/Engine/Extras/P4VUtils/UserInterface.cs
@@ -1,3 +1,5 @@
+// Copyright Epic Games, Inc. All Rights Reserved.
+
 using System;
 using System.Diagnostics;
 using System.IO;
