- BuildGraph p4vutils changes to build and submit to p4 all platform binaries, NFL and regular
- Minor other fixes to p4vutils

diff --git a/Engine/Extras/P4VUtils/BatchFiles/BuildAllPlatforms.bat b/Engine/Extras/P4VUtils/BatchFiles/BuildAllPlatforms.bat
--- a/Engine/Extras/P4VUtils/BatchFiles/BuildAllPlatforms.bat
+++ b/Engine/Extras/P4VUtils/BatchFiles/BuildAllPlatforms.bat
@@ -4,8 +4,11 @@ setlocal
 
 pushd "%~dp0.."
 
+rmdir /s /q bin
+rmdir /s /q obj
+
 rem Compile with defualt WithRestricted value, which may be true
-dotnet publish P4VUtils.csproj /p:IsWindows=true p:IsOSX=false /p:IsLinux=false
+dotnet publish P4VUtils.csproj /p:IsWindows=true /p:IsOSX=false /p:IsLinux=false
 dotnet publish P4VUtils.csproj /p:IsWindows=false /p:IsOSX=true /p:IsLinux=false
 dotnet publish P4VUtils.csproj /p:IsWindows=false /p:IsOSX=false /p:IsLinux=true
 
diff --git a/Engine/Extras/P4VUtils/BatchFiles/BuildAllPlatforms.sh b/Engine/Extras/P4VUtils/BatchFiles/BuildAllPlatforms.sh
--- a/Engine/Extras/P4VUtils/BatchFiles/BuildAllPlatforms.sh
+++ b/Engine/Extras/P4VUtils/BatchFiles/BuildAllPlatforms.sh
@@ -9,6 +9,9 @@ else
 	source ../../Build/BatchFiles/Linux/SetupDotnet.sh
 fi
 
+rm -rf bin
+rm -rf obj
+
 # Compile with defualt WithRestricted value, which may be true
 dotnet publish P4VUtils.csproj -p:IsWindows=true -p:IsOSX=false -p:IsLinux=false -p:EnableWindowsTargeting=true
 dotnet publish P4VUtils.csproj -p:IsWindows=false -p:IsOSX=true -p:IsLinux=false
