[P4VUtils] Rename the 'UE Submit' category to 'UE Content'

#rnx

- There are a few commands that can submit changelists but do not feature under 'UE Submit' submenu which some people felt was misleading.
- As the only command there is currently intended to be used when submitting content, 'UE Content' feels like a good fit.
- Renamed SubmitCommand.cs as this was also misleading.

diff --git a/Engine/Extras/P4VUtils/CommandAttribute.cs b/Engine/Extras/P4VUtils/CommandAttribute.cs
--- a/Engine/Extras/P4VUtils/CommandAttribute.cs
+++ b/Engine/Extras/P4VUtils/CommandAttribute.cs
@@ -12,9 +12,9 @@ namespace P4VUtils
 		Root,
 
 		/// <summary>
-		/// commands that help with submitting files/changelists
+		/// Commands that help with operations relating to content files
 		/// </summary>
-		Submission,
+		Content,
 
 		/// <summary>
 		/// Commands that help with common but simple operations
@@ -22,12 +22,12 @@ namespace P4VUtils
 		Toolbox,
 
 		/// <summary>
-		/// complex commands to facilitate integrations/backout
+		/// Complex commands to facilitate integrations
 		/// </summary>
 		Integrate,
 
 		/// <summary>
-		/// local build and horde preflights
+		/// Local build and horde preflights
 		/// </summary>
 		Horde
 	}
diff --git a/Engine/Extras/P4VUtils/Commands/SubmitAndVirtualizeCommand.cs b/Engine/Extras/P4VUtils/Commands/SubmitAndVirtualizeCommand.cs
new file mode 100644
--- /dev/null
+++ b/Engine/Extras/P4VUtils/Commands/SubmitAndVirtualizeCommand.cs
@@ -0,0 +1,980 @@
+// Copyright Epic Games, Inc. All Rights Reserved.
+
+#if DEBUG
+// When enabled the command will not actually submit the changelist, but will act like it did.
+// This is for development purposes only and so is only available in the debug config.
+//#define UE_DEBUG_DISABLE_SUBMITS
+#endif
+
+using EpicGames.Core;
+using EpicGames.Perforce;
+using Microsoft.Extensions.Logging;
+using Microsoft.Win32;
+using System;
+using System.Collections.Generic;
+using System.IO;
+using System.Linq;
+using System.Text;
+using System.Text.Json;
+using System.Threading;
+using System.Threading.Tasks;
+
+namespace P4VUtils.Commands
+{
+	/// <summary>
+	/// Contains info about a built exe, generated by our build system (found in .version files)
+	/// </summary>
+	public class BuildVersion
+	{
+		public int? MajorVersion { get; set; }
+		public int? MinorVersion { get; set; }
+		public int? PatchVersion { get; set; }
+		public int? Changelist { get; set; }
+		public int? CompatibleChangelist { get; set; }
+		public int? IsLicenseeVersion { get; set; }
+		public int? IsPromotedBuild { get; set; }
+		public string? BranchName { get; set; }
+		public string? BuildId { get; set; }
+	}
+
+	/// <summary>
+	/// A partial representation of the .target file generated by our build system. We are only
+	/// interested in the BuildVersion info and so the class does not have any other members.
+	/// </summary>
+	public class TargetBuildVersion
+	{
+		public BuildVersion? Version { get; set; }
+	}
+
+	public static class LoggingUtils
+	{
+		/// <summary>
+		/// This extension is used to write a new line to the logger. 
+		/// log.NewLine better shows intent over log.LogInformation("")
+		/// </summary>
+		public static void NewLine(this ILogger logger)
+		{
+			logger.LogInformation("");
+		}
+	}
+
+	[Command("submitandvirtualize", CommandCategory.Content, 0)]
+	class SubmitAndVirtualizeCommand : Command
+	{
+		public override string Description => "Virtualize And Submit";
+
+		public override CustomToolInfo CustomTool => new CustomToolInfo("Virtualize And Submit", "$c %C") { ShowConsole = true, RefreshUI = true };
+
+		public override async Task<int> Execute(string[] args, IReadOnlyDictionary<string, string> configValues, ILogger logger)
+		{
+			// Parse command lines
+			if (args.Length < 3)
+			{
+				logger.LogError("Error: Not enough args for command, tool is now exiting");
+				return 1;
+			}
+
+			string clientSpec = args[1];
+
+			List<int> changelistsToSubmit = new List<int>();
+
+			for (int index = 2; index < args.Length; ++index)
+			{
+				if (!int.TryParse(args[index], out int changeNumber))
+				{
+					logger.LogError("Error: '{Argument}' is not a numbered changelist, tool is now exiting", args[index]);
+					return 1;
+				}
+
+				changelistsToSubmit.Add(changeNumber);
+			}
+
+			if (changelistsToSubmit.Count == 0)
+			{
+				logger.LogError("Error: No changelists to submit were provided, tool is now exiting");
+				return 1;
+			}
+
+			// Connect to perforce and validate
+
+			logger.LogInformation("Connecting to Perforce...");
+
+			// We prefer the native client to avoid the problem where different versions of p4.exe expect
+			// or return records with different formatting to each other.
+			PerforceSettings settings = new PerforceSettings(PerforceSettings.Default) { PreferNativeClient = true, ClientName = clientSpec };
+			using IPerforceConnection perforceConnection = await PerforceConnection.CreateAsync(settings, logger);
+			if (perforceConnection == null)
+			{
+				logger.LogError("Error: Failed to connect to Perforce, tool is now exiting");
+				return 1;
+			}
+
+			for (int index = 0; index < changelistsToSubmit.Count; ++index)
+			{
+				logger.NewLine();
+				logger.LogInformation("========== Processing change list {Index}/{Count} ==========", index + 1, changelistsToSubmit.Count);
+
+				int changeNumber = changelistsToSubmit[index];
+				if(await ProcessChangelist(perforceConnection, changeNumber, logger) == false)
+				{
+					logger.LogError("\nError: Failed to process CL {Changelist}, tool is now exiting", changeNumber);
+					return 1;
+				}
+			}
+
+			logger.NewLine();
+			logger.LogInformation("========== All changelist(s) submitted successfully, tool is now exiting ==========");
+			
+			return 0;
+		}
+
+		public static async Task<bool> ProcessChangelist(IPerforceConnection perforceConnection, int changeNumber, ILogger logger)
+		{
+			if (perforceConnection == null || perforceConnection.Settings == null || perforceConnection.Settings.ClientName == null)
+			{
+				logger.LogError("Error: Invalid Perforce connection!");
+				return false;
+			}
+
+			string clientSpec = perforceConnection.Settings.ClientName;
+
+			logger.LogInformation("Attempting to virtualize and submit changelist {Change} in the workspace {Spec}", changeNumber, clientSpec);
+
+			// First we need to find the packages in the changelist
+			string[]? localFilePaths = await FindPackagesInChangelist(perforceConnection, changeNumber, logger);
+			if (localFilePaths == null)
+			{
+				return false;
+			}
+
+			if (localFilePaths.Length > 0)
+			{
+				// If the changelist has shelved files then the submit will fail, so if we have package files that
+				// need to be checked for virtualization we don't want to have the potentially long virtualization
+				// process run (during which the user will alt-tab) only for it to fail at the end.
+				// So we should check for shelved files now and early out before the virtualization process runs.
+				if (await DoesChangelistHaveShelvedFiles(perforceConnection, changeNumber) == true)
+				{
+					logger.LogError("Error: Changelist {Change} has shelved files and cannot be submitted", changeNumber);
+					return false;
+				}
+
+				logger.LogInformation("Found {Amount} package(s) that may need virtualization", localFilePaths.Length);
+
+				// Now sort the package paths by their project (it is unlikely but a user could be submitting content
+				// from multiple projects that the same time)
+
+				Dictionary<string, List<string>> projects = SortPackagesIntoProjects(localFilePaths, logger);
+
+				logger.LogInformation("The packages are distributed between {Amount} project(s)", projects.Count);
+
+				// Find engine per project
+				IReadOnlyDictionary<string, string> engineInstalls = EnumerateEngineInstallations(logger);
+
+				foreach (KeyValuePair<string, List<string>> project in projects)
+				{
+					string projectPath = project.Key;
+					string engineRoot = GetEngineRootForProject(projectPath, logger);
+					if (!String.IsNullOrEmpty(engineRoot))
+					{
+						logger.NewLine();
+						logger.LogInformation("Attempting to virtualize packages in project '{Project}' using the engine installation '{Engine}'", project.Key, engineRoot);
+						// @todo Many projects can share the same engine install, and technically UnrealVirtualizationTool
+						// supports the virtualization files from many projects at the same time. We could consider doing
+						// this pass per engine install rather than per project? At the very least we should only 'build'
+						// the tool once per engine
+
+						Task<bool> compileResult = ValidateOrBuildVirtualizationTool(engineRoot, logger);
+
+						string tempFilesPath = await WritePackageFileList(project.Value, logger);
+
+						// Check if the compilation of the tool succeeded or not
+						if (await compileResult == false)
+						{
+							return false;
+						}
+
+						// Even though this will have been done while we were waiting for BuildVirtualizationTool to complete 
+						// we want to log that it was done after so that the output log makes sense to the user, otherwise 
+						// they will end up thinking that they are waiting on the PackageList to be written rather than on
+						// the tool to be built.
+						logger.LogInformation("PackageList was written to '{Path}'", tempFilesPath);
+
+						if (await RunVirtualizationTool(engineRoot, projectPath, clientSpec, tempFilesPath, logger) == false)
+						{
+							return false;
+						};
+
+						// @todo Make sure this always gets cleaned up when we go out of scope
+						File.Delete(tempFilesPath);
+					}
+					else
+					{
+						logger.LogError("Error: Failed to find engine root for project {Project}", project.Key);
+						return false;
+					}
+				}
+
+				logger.NewLine();
+				logger.LogInformation("All packages have been virtualized");
+
+				//@todo ideally we should get the tags back from UnrealVirtualizationTool
+				ChangeRecord? changeRecord = await StampChangelistDescription(perforceConnection, changeNumber, logger);
+				if (changeRecord == null)
+				{
+					return false;
+				}
+
+				logger.LogInformation("Attempting to submit changelist {Number}...", changeNumber);
+
+#if !UE_DEBUG_DISABLE_SUBMITS
+				if (await SubmitChangelist(perforceConnection, changeNumber, logger) == false)
+				{
+					// If the final submit failed we remove the virtualization tag, even though the changelist is technically
+					// virtualized at this point and submitting it would be safe.
+					// This is to keep the behavior the same as the native code paths.
+					logger.LogInformation("Removing virtualization tags from the changelist...");
+					PerforceResponse updateResponse = await perforceConnection.TryUpdateChangeAsync(UpdateChangeOptions.None, changeRecord, CancellationToken.None);
+					if (!updateResponse.Succeeded)
+					{
+						logger.LogError("Error: Failed to remove the virtualization tags!");
+					}
+					logger.LogInformation("Virtualization tags have been removed.");
+
+					return false;
+				}
+#else
+				// Dummy pretend submit for debugging purposes, allows developer iteration without actually submitting
+				logger.LogInformation("Successfully submited changelist {Change}", changeNumber);
+#endif // !UE_DEBUG_DISABLE_SUBMITS
+			}
+			else
+			{
+				logger.LogInformation("The changelist does not contain any package files, submitting as normal...");
+
+				if (await SubmitChangelist(perforceConnection, changeNumber, logger) == false)
+				{
+					return false;
+				}
+			}
+		
+			return true;
+		}
+
+		/// <summary>
+		/// Checks to see if a precompiled version of the UnrealVirtualizationTool exists and if so
+		/// if it is valid for us to use this.
+		/// If we cannot use a precompiled version of the tool then we will request that it be 
+		/// compiled locally.
+		/// </summary>
+		/// <param name="engineRoot">Root path of the engine we want to build the tool for</param>
+		/// <param name="logger">Interface for logging</param>
+		/// <returns>True if the UnrealVirtualizationTool.exe exists in a form that we can use, otherwise false</returns>
+		private static async Task<bool> ValidateOrBuildVirtualizationTool(string engineRoot, ILogger logger)
+		{
+			// Check to see if the precompiled binaries for the  editor and virtualization tool are available
+
+			if (!IsSourceCodeAvaliable(engineRoot))
+			{
+				if (DoesEngineToolExist(engineRoot, "UnrealVirtualizationTool"))
+				{
+					logger.LogInformation("Using a precompiled binary version of UnrealVirtualizationTool");
+					return true;
+				}
+				else
+				{
+					logger.LogError("Error: No source code and no precompiled binary of UnrealVirtualizationTool found");
+					return false;
+				}
+			}
+			else
+			{
+				// Since we have no good way to determine if an exe was from the PCB or not, if we have
+				// source code we will just have to build the tool anyway for safety.
+				// Once this is fixed we should use the old logic which was to use the PCB version of
+				// UnrealVirtualizationTool if the editor was a PCB version and only to recompile it 
+				// once we detect that the editor has been compiled locally.
+
+				// We must try to build the tool locally, if the user does not have a valid code compilation 
+				// tool chain installed then UnrealBuildTool will give them errors.
+				return await BuildVirtualizationTool(engineRoot, logger);
+			}
+		}
+
+		/// <summary>
+		/// Checks to see if the user has source code synced or not.
+		/// To make this assumption we check to see if we can find the UnrealBuildTool project. If we can
+		/// then it is a fair bet that the user can compile code, if they do not have the project then
+		/// they probably can't.
+		/// </summary>
+		/// <param name="engineRoot">Root path of the engine we want to build the tool for</param>
+		/// <returns>True if source code is present, otherwise false</returns>
+		private static bool IsSourceCodeAvaliable(string engineRoot)
+		{
+			string versionPath = String.Format(@"{0}\Engine\Source\Programs\UnrealBuildTool\UnrealBuildTool.csproj", engineRoot);
+
+			if (System.IO.File.Exists(versionPath))
+			{
+				return true;
+			}
+			else
+			{
+				return false;
+			}
+		}
+
+		/// <summary>
+		/// Checks if the given tool exists in the engine binaries or not
+		/// </summary>
+		/// <param name="engineRoot">Root path of the engine we want to build the tool for</param>
+		/// <param name="toolName">The name of the tool to look for</param>
+		/// <returns>TRue if the exe for the tool already exists, otherwise false</returns>
+		private static bool DoesEngineToolExist(string engineRoot, string toolName)
+		{
+			string toolPath = String.Format(@"{0}\Engine\Binaries\Win64\{1}.exe", engineRoot, toolName);
+
+			if (System.IO.File.Exists(toolPath))
+			{
+				return true;
+			}
+			else
+			{
+				return false;
+			}
+		}
+
+		/// <summary>
+		/// Compiles UnrealVirtualizationTool via RunUBT.bat
+		/// </summary>
+		/// <param name="engineRoot">Root path of the engine we want to build the tool for</param>
+		/// <param name="logger">Interface for logging</param>
+		/// <returns>True if the tool built successfully, otherwise false</returns>
+		private static async Task<bool> BuildVirtualizationTool(string engineRoot, ILogger logger)
+		{
+			logger.LogInformation("Building UnrealVirtualizationTool...");
+
+			string buildBatchFile = Path.Combine(engineRoot, @"Engine\Build\BatchFiles\RunUBT.bat");
+			StringBuilder arguments = new StringBuilder($"{buildBatchFile.QuoteArgument()}");
+			arguments.Append(" UnrealVirtualizationTool Win64 development -Progress");
+
+			string shellFileName = Environment.GetEnvironmentVariable("COMSPEC") ?? "C:\\Windows\\System32\\cmd.exe";
+			string shellArguments = $"/d/s/c \"{arguments}\"";
+
+			using (MemoryStream bufferedOutput = new MemoryStream())
+			using (ManagedProcessGroup Group = new ManagedProcessGroup())
+			using (ManagedProcess Process = new ManagedProcess(Group, shellFileName, shellArguments, null, null, System.Diagnostics.ProcessPriorityClass.Normal))
+			{
+				await Process.CopyToAsync(bufferedOutput, CancellationToken.None);
+
+				// We only show the output if there was an error to avoid showing too much
+				// info to the user when they are invoking this from p4v.
+				// We need to print everything from stdout and stderr as stderr alone often 
+				// does not contain all of the build failure info. (note that ManagedProcess is
+				// merging the two streams by default)
+				if (Process.ExitCode != 0)
+				{
+					bufferedOutput.Seek(0, SeekOrigin.Begin);
+
+					using (Stream stdOutput = Console.OpenStandardOutput())
+					{
+						await bufferedOutput.CopyToAsync(stdOutput, CancellationToken.None);
+					}
+
+					logger.LogError("Error: Failed to build UnrealVirtualizationTool");
+					return false;
+				}
+			}
+
+			logger.LogInformation("UnrealVirtualizationTool built successfully");
+			return true;
+		}
+
+		/// <summary>
+		/// Runs the UnrealVirtualizationTool
+		/// </summary>
+		/// <param name="engineRoot">Root path of the engine we want to run the tool from</param>
+		/// <param name="projectPath">The absolute path of the project file for the packages in the package list</param>
+		/// <param name="clientSpec">The perforce client spec that the files are in</param>
+		/// <param name="packageListPath">A path to a text file containing the paths of the packages to be virtualized</param>
+		/// <param name="logger">Interface for logging</param>
+		/// <returns></returns>
+		private static async Task<bool> RunVirtualizationTool(string engineRoot, string projectPath, string clientSpec, string packageListPath, ILogger logger)
+		{
+			logger.LogInformation("Running UnrealVirtualizationTool...");
+
+			string toolPath = Path.Combine(engineRoot, @"Engine\Binaries\Win64\UnrealVirtualizationTool.exe");
+			string toolArgs = string.Format("\"{0}\" -MinimalLogging -ClientSpecName={1} -Mode=PackageList -Path=\"{2}\"", projectPath, clientSpec, packageListPath);
+
+			using (Stream stdOutput = Console.OpenStandardOutput())
+			using (ManagedProcessGroup Group = new ManagedProcessGroup())
+			using (ManagedProcess Process = new ManagedProcess(Group, toolPath, toolArgs, null, null, System.Diagnostics.ProcessPriorityClass.Normal))
+			{
+				await Process.CopyToAsync(stdOutput, CancellationToken.None);
+
+				if (Process.ExitCode != 0)
+				{
+					logger.LogError("Error: UnrealVirtualizationTool failed!");
+					return false;
+				}
+			}
+
+			return true;
+		}
+
+		/// <summary>
+		/// Writes out a list of package paths to a text file stored in the users 
+		/// temp directory.
+		/// The name of the file will be a randomly generated GUID making file name
+		/// collisions unlikely.
+		/// Each path will be written to it's own line inside of the files
+		/// </summary>
+		/// <param name="packagePaths">A list of package file paths to be written to the file</param>
+		/// <param name="logger">Interface for logging</param>
+		/// <returns>The path of the file once it has been written</returns>
+		private static async Task<string> WritePackageFileList(List<string> packagePaths, ILogger logger)
+		{
+			// We pass the list of packages to the tool via a file as the number of package 
+			// paths can potentially be huge and exceed the cmdline length. 
+			// So currently we write the files to a UnrealVirtualizationTool directory under the temp directory
+			string tempDirectory = Path.Combine(Path.GetTempPath(), "UnrealVirtualizationTool");
+			Directory.CreateDirectory(tempDirectory);
+
+			string tempFilesPath = Path.Combine(tempDirectory, Guid.NewGuid().ToString() + ".txt");
+
+			using (StreamWriter Writer = new StreamWriter(tempFilesPath))
+			{
+				foreach (string line in packagePaths)
+				{
+					await Writer.WriteLineAsync(line);
+				}
+			}
+
+			return NormalizeFilename(tempFilesPath);
+		}
+
+		/// <summary>
+		/// Sorts the given package paths by the unreal project that they are found to be in.
+		/// If a package is found to not be in a project it will currently raise a warning 
+		/// but not prevent further execution.
+		/// </summary>
+		/// <param name="packagePaths">A list of absolute file paths pointing to package files</param>
+		/// <param name="logger">Interface for logging</param>
+		/// <returns>A dictionary where the key is the path of an unreal project and the value is a list of packages in that project</returns>
+		private static Dictionary<string, List<string>> SortPackagesIntoProjects(string[] packagePaths, ILogger logger)
+		{
+			Dictionary<string, List<string>> projects = new Dictionary<string, List<string>>();
+			foreach (string path in packagePaths)
+			{
+				string normalizedPath = NormalizeFilename(path);
+				string projectFilePath = FindProjectForPackage(normalizedPath, logger);
+
+				if (!String.IsNullOrEmpty(projectFilePath))
+				{
+					if (!projects.ContainsKey(projectFilePath))
+					{
+						projects.Add(projectFilePath, new List<string>());
+					}
+
+					projects[projectFilePath].Add(normalizedPath);
+				}
+				else
+				{
+					// @todo Re-evaluate if we want this warning at some point in the future. 
+					// Technically submitting a package file not under a project could have valid use cases
+					// but if the user called this command it would indicate that they expect that they are
+					// submitting packages in a valid project.
+					// So for now we give a warning so that it is easier to catch cases where packages are
+					// thought not to be part of a project even if they are.
+					logger.LogWarning("Unable to find a valid project for the package '{Path}'", normalizedPath);
+				}
+			}
+
+			return projects;
+		}
+
+		private static async Task<bool> SubmitChangelist(IPerforceConnection perforceConnection, int changeNumber, ILogger logger)
+		{
+			PerforceResponseList<SubmitRecord> submitResponses = await TrySubmitAsync(perforceConnection, changeNumber, SubmitOptions.None, CancellationToken.None);
+
+			bool successfulSubmit = submitResponses.All(x => x.Succeeded);
+			if (successfulSubmit)
+			{
+				// The submit request will return a number of records. One for the original changelist (null)
+				// one for each submitted file (all null) and the last record detailing the submitted changelist.
+				// So we can just grab the submitted changelist number from the last record.
+
+				PerforceResponse<SubmitRecord> submittedResponse = submitResponses[submitResponses.Count - 1];
+				logger.LogInformation("Successfully submited CL {SrcCL} as CL {DstCL}", changeNumber, submittedResponse.Data.SubmittedChangeNumber);
+				
+				return true;
+			}
+			else
+			{
+				// Log every response that was a failure and has an error message associated with it
+				logger.LogError("Error: Submit failed due to:");
+				foreach (PerforceResponse<SubmitRecord> response in submitResponses)
+				{
+					if (response.Failed && response.Error != null)
+					{
+						logger.LogError("\t{Message}", response.Error.Data.ToString());
+					}
+				}
+
+				return false;
+			}
+		}
+
+		/// <summary>
+		/// Adds the virtualization tags to a changelist description which is used to show that the virtualization
+		/// process has been run on it, before it was submitted.
+		/// </summary>
+		/// <param name="perforceConnection">A valid connection to perforce to use</param>
+		/// <param name="changeNumber">The changelist we should be stamping</param>
+		/// <param name="logger">Interface for logging</param>
+		/// <returns>If the stamp succeeded then it returns a ChangeRecord representing the changelist as it was before it was stamped, returns null on failure</returns>
+		private static async Task<ChangeRecord?> StampChangelistDescription(IPerforceConnection perforceConnection, int changeNumber, ILogger logger)
+		{
+			logger.LogInformation("Adding virtualization tags to changelist description...");
+
+			ChangeRecord changeRecord;
+			try
+			{
+				changeRecord = await perforceConnection.GetChangeAsync(GetChangeOptions.None, changeNumber);
+			}
+			catch (Exception)
+			{
+				logger.LogError("Error: Failed to get the description {Change} so we can edit it", changeNumber);
+				return null;
+			}
+
+			string? originalDescription = changeRecord.Description;
+
+			// @todo Should be getting the tag from the virtualization tool!
+			changeRecord.Description += "\n#virtualized\n";
+
+			PerforceResponse updateResponse = await perforceConnection.TryUpdateChangeAsync(UpdateChangeOptions.None, changeRecord, CancellationToken.None);
+			if (updateResponse.Succeeded)
+			{
+				// Restore the original description so we are returning the original ChangeRecord
+				changeRecord.Description = originalDescription;
+				return changeRecord;
+			}
+			else
+			{
+				logger.LogError("Error: Failed to edit the description of {Change} due to\n{Message}", changeNumber, updateResponse.Error!.ToString());
+				return null;
+			}
+		}
+
+		/// <summary>
+		/// Find all of the unreal packages in a single perforce changelist and return them as local file paths on the users machine.
+		/// </summary>
+		/// <param name="perforceConnection">A valid connection to perforce. This should have the correct client spec for the given changelist number</param>
+		/// <param name="changeNumber">The changelist we should look in</param>
+		/// <param name="logger">Interface for logging</param>
+		/// <returns>A list of all the packages in the given changelist, in local file path format</returns>
+		private static async Task<string[]?> FindPackagesInChangelist(IPerforceConnection perforceConnection, int changeNumber, ILogger logger)
+		{
+			logger.LogInformation("Finding files in changelist {Change}...", changeNumber);
+			DescribeRecord changeRecord;
+			try
+			{
+				changeRecord = await perforceConnection.DescribeAsync(changeNumber, CancellationToken.None);
+			}
+			catch (Exception)
+			{
+				logger.LogError("Error: Failed to find changelist {Change}", changeNumber);
+				return null;
+			}
+
+			if (changeRecord.Files.Count == 0)
+			{
+				logger.LogError("Error: Changelist {Change} is empty, cannot submit", changeNumber);
+				return null;
+			}
+
+			// @todo Should we check if the changelist has shelved files and error at this point since
+			// we know that the user will not be able to submit it?
+
+			// Find the depth paths in the changelist that point to package files
+			string[] depotPackagePaths = changeRecord.Files.Select(x => x.DepotFile).Where(x => IsPackagePath(x)).ToArray();
+
+			// We can early out if there are no packages in the changelist
+			if (depotPackagePaths.Length == 0)
+			{
+				return depotPackagePaths;
+			}
+
+			// Now convert from depot paths to local paths on the users machine
+			List<WhereRecord> whereRecords = await perforceConnection.WhereAsync(depotPackagePaths, CancellationToken.None).ToListAsync();
+
+			return whereRecords.Select(x => x.Path).ToArray();
+		}
+
+		/// <summary>
+		/// Returns if a changelist has shelved files or not
+		/// </summary>
+		/// <param name="perforceConnection">A valid connection to perforce. This should have the correct client spec for the given changelist number</param>
+		/// <param name="changeNumber">The changelist we should look in</param>
+		/// <returns>True if the changelist contains shelved files, otherwise false</returns>
+		private static async Task<bool> DoesChangelistHaveShelvedFiles(IPerforceConnection perforceConnection, int changeNumber)
+		{
+			List<DescribeRecord> responses = await perforceConnection.DescribeAsync(DescribeOptions.Shelved, -1, new int[] { changeNumber }, CancellationToken.None);
+
+			if (responses.Count == 1)
+			{
+				return responses[0].Files.Count != 0;
+			}
+			else
+			{
+				return false;
+			}
+		}
+
+		/// <summary>
+		/// Finds the unreal project file for a given unreal package
+		/// </summary>
+		/// <param name="packagePath">The package to find the project for</param>
+		/// <param name="logger">Interface for logging</param>
+		/// <returns>The path of the projects .uproject file if found, an empty string if no valid project file was found</returns>
+		private static string FindProjectForPackage(string packagePath, ILogger logger)
+		{
+			// @todo note that this function mirrors FUnrealVirtualizationToolApp::TryFindProject
+			// both will not work correctly with some plugin setups, this is known and will be
+			// fixed later when FUnrealVirtualizationToolApp is also fixed.
+			packagePath = NormalizeFilename(packagePath);
+
+			int contentIndex = packagePath.LastIndexOf("/content/", StringComparison.OrdinalIgnoreCase);
+			if (contentIndex == -1)
+			{
+				logger.LogWarning("'{Path}' is not under a content directory", packagePath);
+				return string.Empty;
+			}
+
+			while (contentIndex != -1)
+			{
+				// Assume that the project directory is the parent of the /content/ directory
+				string projectDirectory = packagePath[..contentIndex];
+
+				string[] projectFiles = Directory.GetFiles(projectDirectory, "*.uproject");
+
+				if (projectFiles.Length == 0)
+				{
+					// If there was no project file, the package could be in a plugin, so lets check for that
+					string pluginDirectory = projectDirectory;
+					string[] pluginFiles = Directory.GetFiles(pluginDirectory, "*.uplugin");
+
+					if (pluginFiles.Length == 1)
+					{
+						// We have a valid plugin file, so we should be able to find a /plugins/ directory which will be just below the project directory
+						int pluginIndex = pluginDirectory.LastIndexOf("/plugins/", StringComparison.OrdinalIgnoreCase);
+						if (pluginIndex != -1)
+						{
+							// We found the plugin root directory so the one above it should be the project directory
+							projectDirectory = pluginDirectory[..pluginIndex];
+							projectFiles = Directory.GetFiles(projectDirectory, "*.uproject");
+						}
+					}
+					else if (pluginFiles.Length > 1)
+					{
+						logger.LogWarning("Found multiple .uplugin files for '{Path}' at '{PluginDir}'", pluginDirectory);
+						return string.Empty;
+					}
+				}
+
+				if (projectFiles.Length == 1)
+				{
+					return NormalizeFilename(projectFiles[0]);
+				}
+				else if (projectFiles.Length > 1)
+				{
+					logger.LogWarning("Found multiple .uproject files for '{Path}' at '{ProjectDir}'", projectDirectory);
+					return string.Empty;
+				}
+
+				// Could be more than one content directory in the path so lets keep looking
+				contentIndex = packagePath.LastIndexOf("/content/", contentIndex, StringComparison.OrdinalIgnoreCase);
+			}
+
+			// We found one or more content directories but none of them contained a project file
+			logger.LogWarning("FFailed to find project file for '{Package}'", packagePath);
+			return string.Empty;
+		}
+
+		private static string GetEngineRootForProject(string projectFilePath, ILogger logger)
+		{
+			string engineIdentifier = GetEngineIdentifierForProject(projectFilePath, logger);
+			if (!String.IsNullOrEmpty(engineIdentifier))
+			{
+				string engineRoot = GetEngineRootDirFromIdentifier(engineIdentifier, logger);
+				if (!String.IsNullOrEmpty(engineRoot))
+				{
+					return engineRoot;
+				}
+				else 
+				{
+					logger.LogWarning("Unable to find an engine root for installation {Identifier}, will attempt to find the engine via the directory hierarchy", engineIdentifier);
+				}	
+			}
+
+			return FindEngineFromPath(projectFilePath);
+		}
+
+		// The following functions mirror code found @Engine\Source\Runtime\CoreUObject\Private\Misc\PackageName.cpp
+		#region CoreUObject PackageName 
+
+		private static bool IsPackagePath(string path)
+		{
+			return IsBinaryPackagePath(path) || IsTextPackagePath(path);
+		}
+		
+		private static bool IsBinaryPackagePath(string path)
+		{
+			return path.EndsWith(".uasset", StringComparison.OrdinalIgnoreCase) || path.EndsWith(".umap", StringComparison.OrdinalIgnoreCase);
+		}
+
+		private static bool IsTextPackagePath(string path)
+		{
+			return path.EndsWith(".utxt", StringComparison.OrdinalIgnoreCase) || path.EndsWith(".utxtmap", StringComparison.OrdinalIgnoreCase);
+		}
+
+		#endregion
+
+		// The following functions mirror code found @Engine\Source\Runtime\Core\Private\Misc\Paths.cpp
+		#region Paths
+
+		private static string NormalizeFilename(string path)
+		{
+			return path.Replace('\\', '/');
+		}
+
+		private static string NormalizeDirectoryName(string path)
+		{
+			path = path.Replace('\\', '/');
+
+			if (path.EndsWith('/'))
+			{
+				path = path.Remove(path.Length - 1, 1);
+			}
+
+			return path;
+		}
+
+		#endregion
+
+		// The following functions mirror code found @Engine\Source\Developer\DesktopPlatform\Private\DesktopPlatformBase.cpp
+		#region DesktopPlatform DesktopPlatformBase
+
+		// Note that for the C# versions we only support the modern way of associating
+		// projects with engine installation, via the 'EngineAssociation' entry in
+		// the .uproject file.
+
+		private static string GetEngineIdentifierForProject(string projectFilePath, ILogger logger)
+		{
+			try
+			{
+				JsonObject root = JsonObject.Read(new FileReference(projectFilePath));
+				string engineIdentifier = root.GetStringField("EngineAssociation");
+				if (!String.IsNullOrEmpty(engineIdentifier))
+				{
+					// @todo what if it is a path? (native code has support for this
+					// possibly for an older version)
+					return engineIdentifier;
+				}
+			}
+			catch (Exception ex)
+			{
+				logger.LogError("Error: Failed to parse {File} to find the engine association due to: {Reason}", projectFilePath, ex.Message);
+			}
+
+			// @todo In the native version if there is no identifier we will try to
+			// find the engine root in the directory hierarchy then either find it
+			// identifier or register one if needed.
+
+			return string.Empty;
+		}
+
+		private static string FindEngineFromPath(string path)
+		{
+			string? directoryName = Path.GetDirectoryName(path);
+			
+			if (!String.IsNullOrEmpty(directoryName))
+			{
+				DirectoryInfo? directoryToSearch = new DirectoryInfo(directoryName);
+
+				while (directoryToSearch != null)
+				{
+					if (IsValidRootDirectory(directoryToSearch.ToString()))
+					{
+						return NormalizeDirectoryName(directoryToSearch.ToString());
+					}
+
+					directoryToSearch = Directory.GetParent(directoryToSearch.ToString());
+				}
+			}
+
+			return string.Empty;
+		}
+
+		private static string GetEngineRootDirFromIdentifier(string engineIdentifier, ILogger logger)
+		{
+			IReadOnlyDictionary<string, string> engineInstalls = EnumerateEngineInstallations(logger);
+
+			if (engineInstalls.TryGetValue(engineIdentifier, out string? engineRoot))
+			{
+				return engineRoot;
+			}
+			else
+			{
+				return string.Empty;
+			}
+		}
+
+		private static string GetEngineIdentifierFromRootDir(string rootDirectory, ILogger logger)
+		{
+			rootDirectory = NormalizeDirectoryName(rootDirectory);
+			
+			IReadOnlyDictionary<string, string> engineInstalls = EnumerateEngineInstallations(logger);
+
+			foreach (KeyValuePair<string, string> pair in engineInstalls)
+			{
+				if (String.Equals(pair.Value,rootDirectory, StringComparison.OrdinalIgnoreCase))
+				{
+					return pair.Key;
+				}
+			}
+
+			return String.Empty;
+		}
+
+		private static bool IsValidRootDirectory(string rootDirectory)
+		{
+			// Check that there's an Engine\Binaries directory underneath the root
+			string engineBinariesDirectory = Path.Combine(rootDirectory, "Engine/Binaries");
+
+			if (!Directory.Exists(engineBinariesDirectory))
+			{
+				return false;
+			}
+
+			// Also check there's an Engine\Build directory. This will filter out anything
+			// //that has an engine-like directory structure but doesn't allow building
+			// code projects - like the launcher.
+
+			string engineBuildDirectory = Path.Combine(rootDirectory, "Engine/Build");
+
+			if (!Directory.Exists(engineBuildDirectory))
+			{
+				return false;
+			}
+
+			return true;
+		}
+
+		private static IReadOnlyDictionary<string, string> EnumerateEngineInstallations(ILogger logger)
+		{
+			Dictionary<string, string> launcherEngineInstalls = EnumerateLauncherEngineInstallations(logger);
+			Dictionary<string, string> customEngineInstalls = EnumerateCustomEngineInstallations(logger);
+
+			Dictionary<string, string> allInstalls = new Dictionary<string, string>();
+
+			foreach (KeyValuePair<string, string> pair in launcherEngineInstalls)
+			{
+				allInstalls.Add(pair.Key, pair.Value);
+			}
+
+			foreach (KeyValuePair<string, string> pair in customEngineInstalls)
+			{
+				allInstalls.Add(pair.Key, pair.Value);
+			}
+
+			return allInstalls;
+		}
+
+		private static Dictionary<string, string> EnumerateLauncherEngineInstallations(ILogger logger)
+		{
+			Dictionary<string, string> installations = new Dictionary<string, string>();
+
+			try
+			{
+				string installedListFilePath = Path.Combine(Environment.GetFolderPath(Environment.SpecialFolder.CommonApplicationData), "Epic/UnrealEngineLauncher/LauncherInstalled.dat");
+
+				if (JsonObject.TryRead(new FileReference(installedListFilePath), out JsonObject? rawObject))
+				{
+					JsonObject[] installationArray = rawObject.GetObjectArrayField("InstallationList");
+					foreach (JsonObject installObject in installationArray)
+					{
+						string appName = installObject.GetStringField("AppName");
+						if (appName.StartsWith("UE_", StringComparison.Ordinal))
+						{
+							appName = appName.Remove(0, 3);
+							string installPath = NormalizeDirectoryName(installObject.GetStringField("InstallLocation"));
+
+							installations.Add(appName, installPath);
+						}
+					}
+				}
+			}
+			catch (Exception ex)
+			{
+				logger.LogWarning("EnumerateLauncherEngineInstallations: {Message}", ex.Message);
+				installations.Clear();
+			}
+
+			return installations;
+		}
+
+		private static Dictionary<string, string> EnumerateCustomEngineInstallations(ILogger logger)
+		{
+			Dictionary<string, string> installations = new Dictionary<string, string>();
+
+			if (OperatingSystem.IsWindows())
+			{
+				try
+				{
+					using RegistryKey? subKey = Registry.CurrentUser.OpenSubKey("SOFTWARE\\Epic Games\\Unreal Engine\\Builds", false);
+					if (subKey != null)
+					{
+						foreach (string installName in subKey.GetValueNames())
+						{
+							string? installPath = subKey.GetValue(installName) as string;
+							if (!String.IsNullOrEmpty(installPath))
+							{
+								installations.Add(installName, NormalizeDirectoryName(installPath));
+							}
+						}
+					}
+				}
+				catch (Exception ex)
+				{
+					logger.LogWarning("EnumerateCustomEngineInstallations: {Message}", ex.Message);
+					installations.Clear();
+				}
+			}
+
+			return installations;
+		}
+
+		#endregion
+
+		// The following code acts as an extension to code found @Engine\Source\Programs\Shared\EpicGames.Perforce\PerforceConnection.cs
+		#region PerforceConnection
+
+		// @todo None of the submit functions in PerforceConnection.cs seem to report submit errors in a way that we can inform the user
+		// so this is a custom version that returns the entire response list as different submit errors will return different error
+		// responses. We should fix the API in PerforceConnection so that other areas of code can get better submit error reporting.
+		private static async Task<PerforceResponseList<SubmitRecord>> TrySubmitAsync(IPerforceConnection connection, int changeNumber, SubmitOptions options, CancellationToken cancellationToken = default)
+		{
+			List<string> arguments = new List<string>();
+			if ((options & SubmitOptions.ReopenAsEdit) != 0)
+			{
+				arguments.Add("-r");
+			}
+			arguments.Add($"-c{changeNumber}");
+
+			return (await connection.CommandAsync<SubmitRecord>("submit", arguments, null, cancellationToken));
+		}
+
+		#endregion
+	}
+}
diff --git a/Engine/Extras/P4VUtils/Commands/SubmitCommand.cs b/Engine/Extras/P4VUtils/Commands/SubmitCommand.cs
deleted file mode 100644
diff --git a/Engine/Extras/P4VUtils/Program.cs b/Engine/Extras/P4VUtils/Program.cs
--- a/Engine/Extras/P4VUtils/Program.cs
+++ b/Engine/Extras/P4VUtils/Program.cs
@@ -464,7 +464,7 @@ namespace P4VUtils
 			if (bInstall)
 			{
 				InstallCommandsListInFolder("UE RootHelpers", false/*AddFolderToContextMenu*/, CommandCategory.Root, Document, ExecutableLocation, Logger);
-				InstallCommandsListInFolder("UE Submit", true/*AddFolderToContextMenu*/, CommandCategory.Submission, Document, ExecutableLocation, Logger);
+				InstallCommandsListInFolder("UE Content", true/*AddFolderToContextMenu*/, CommandCategory.Content, Document, ExecutableLocation, Logger);
 				InstallCommandsListInFolder("UE Toolbox", true/*AddFolderToContextMenu*/, CommandCategory.Toolbox, Document, ExecutableLocation, Logger);
 				InstallCommandsListInFolder("UE Integrate", true/*AddFolderToContextMenu*/, CommandCategory.Integrate, Document, ExecutableLocation, Logger);
 				InstallCommandsListInFolder("UE Horde", true/*AddFolderToContextMenu*/, CommandCategory.Horde, Document, ExecutableLocation, Logger);
