[P4VUtils] Add a new command to open jira issues tagged in a changelist description in the callers web browser.

- Add a new command 'OpenJiraCommand' which will attempt to open jira in the users web browser to any issue tagged in that changelists description via the '#jira' tag.
-- We assume that the issue name will be Project-IssueNumber, where the project can contain both letters and numbers. The jira standard enforces that the first character of the project name cannot be a number but we have not enforced that in our checks for simplicity.
-- If multiple issues are tagged we should be able to match with them, even if they are on the same line and/or separated with ' ', ',' or '+'
-- If no tags are found we inform the user via a popup message.
-- We now have an entry in the ini file for the jira server.

- Added a new CommandCategory named 'Browser' to be used by commands that attempt to open new pages in the users web browser.
-- Moved the existing OpenPreflightCommand and RobomergePreviewCommand commands there.
-- It could be argued that PreflightCommand should be moved to 'Browser' too but it feels more at home under 'Horde'
- A number of commands were copy/pasting the code to open a path in a new process (given that the syntax is slightly different per platform) so this code has been moved to a shared ProcessUtils class.

diff --git a/Engine/Extras/P4VUtils/CommandAttribute.cs b/Engine/Extras/P4VUtils/CommandAttribute.cs
--- a/Engine/Extras/P4VUtils/CommandAttribute.cs
+++ b/Engine/Extras/P4VUtils/CommandAttribute.cs
@@ -29,7 +29,12 @@ namespace P4VUtils
 		/// <summary>
 		/// Local build and horde preflights
 		/// </summary>
-		Horde
+		Horde,
+
+		/// <summary>
+		/// Commands that open pages in the users browser
+		/// </summary>
+		Browser
 	}
 
 	/// <summary>
diff --git a/Engine/Extras/P4VUtils/Commands/DescribeCommand.cs b/Engine/Extras/P4VUtils/Commands/DescribeCommand.cs
--- a/Engine/Extras/P4VUtils/Commands/DescribeCommand.cs
+++ b/Engine/Extras/P4VUtils/Commands/DescribeCommand.cs
@@ -58,22 +58,6 @@ namespace P4VUtils.Commands
 		private int MaxDaysHistory = 30;
 		private int MaxChangeListCount = 10000;
 
-		static void OSOpenFile(string FilePath)
-		{
-			if (RuntimeInformation.IsOSPlatform(OSPlatform.Windows))
-			{
-				Process.Start(new ProcessStartInfo(FilePath) { UseShellExecute = true });
-			}
-			else if (RuntimeInformation.IsOSPlatform(OSPlatform.Linux))
-			{
-				Process.Start("xdg-open", FilePath);
-			}
-			else if (RuntimeInformation.IsOSPlatform(OSPlatform.OSX))
-			{
-				Process.Start("open", FilePath);
-			}
-		}
-
 		public override async Task<int> Execute(string[] Args, IReadOnlyDictionary<string, string> ConfigValues, ILogger Logger)
 		{
 			if (Args.Length < 2)
@@ -182,7 +166,7 @@ namespace P4VUtils.Commands
 				await Writer.WriteLineAsync(CSVFormat.ToString());
 			}
 
-			OSOpenFile(PlainTextPath);
+			ProcessUtils.OpenInNewProcess(PlainTextPath);
 
 			Logger.LogInformation("Complete");
 			Logger.LogInformation("PlainText: {PlainText}", PlainTextPath);
diff --git a/Engine/Extras/P4VUtils/Commands/JiraCommands.cs b/Engine/Extras/P4VUtils/Commands/JiraCommands.cs
new file mode 100644
--- /dev/null
+++ b/Engine/Extras/P4VUtils/Commands/JiraCommands.cs
@@ -0,0 +1,93 @@
+// Copyright Epic Games, Inc. All Rights Reserved.
+
+using EpicGames.Perforce;
+using Microsoft.Extensions.Logging;
+using System.Collections.Generic;
+using System.Threading.Tasks;
+using System.Text.RegularExpressions;
+
+namespace P4VUtils.Commands
+{
+
+	[Command("openjira", CommandCategory.Browser, 2)]
+	class OpenJiraCommand : Command
+	{
+		public override string Description => "Jira: Open issues in browser...";
+
+		public override CustomToolInfo CustomTool => new CustomToolInfo("Jira: Open issues in browser", "%C");
+
+		public override async Task<int> Execute(string[] args, IReadOnlyDictionary<string, string> configValues, ILogger logger)
+		{
+			logger.LogInformation("Parsing args ...");
+
+			// Parse command lines
+			if (args.Length < 2)
+			{
+				logger.LogError("Not enough args for command, tool is now exiting");
+				return 1;
+			}
+
+			if (!int.TryParse(args[1], out int changeNumber))
+			{
+				logger.LogError("'{Argument}' is not a numbered changelist, tool is now exiting", args[1]);
+				return 1;
+			}
+
+			logger.LogInformation("Connecting to Perforce...");
+
+			PerforceSettings settings = new PerforceSettings(PerforceSettings.Default);
+
+			using IPerforceConnection perforceConnection = await PerforceConnection.CreateAsync(settings, logger);
+			if (perforceConnection == null)
+			{
+				logger.LogError("Failed to connect to Perforce, tool is now exiting");
+				return 1;
+			}
+
+			PerforceResponse<DescribeRecord> response = await perforceConnection.TryDescribeAsync(changeNumber);
+			if (!response.Succeeded)
+			{
+				logger.LogError("Failed to find the changelist description for {CLNumber} due to {Error}", changeNumber, response.Error?.ToString() ?? "Unknown Error");
+				return 1;
+			}
+
+			DescribeRecord description = response.Data;
+			MatchCollection matches = Regex.Matches(description.Description, @"(?<=#jira.*?)([\w]+-[\d]+)", RegexOptions.Multiline);
+
+			if(matches.Count == 0)
+			{
+				string message = $"Description for {changeNumber} does not contain any valid jira issue tags (#jira ...)";
+
+				logger.LogInformation("{Message}", message);
+				UserInterface.ShowSimpleDialog(message, "No Jira Issues Found", logger);
+				return 1;
+			}
+
+			logger.LogInformation("Found {NumMatches} jira issues in the description for CL {CLNumber}", matches.Count, changeNumber);
+			string jiraServer = GetJiraServerAddress(configValues);
+
+			foreach (Match? match in matches)
+			{
+				string jiraUrl = jiraServer + "/" + match!.Groups[1].Value;
+
+				logger.LogInformation("Opening webpage for {URL}", jiraUrl);
+				ProcessUtils.OpenInNewProcess(jiraUrl);
+			}
+
+			logger.LogInformation("Command completed successfully");
+
+			return 0;
+		}
+
+		private static string GetJiraServerAddress(IReadOnlyDictionary<string, string> ConfigValues)
+		{
+			string? BaseUrl;
+			if (!ConfigValues.TryGetValue("JiraServer", out BaseUrl))
+			{
+				BaseUrl = "https://configure-server-url-in-p4vutils.ini";
+			}
+
+			return BaseUrl.TrimEnd('/');
+		}
+	}
+}
\ No newline at end of file
diff --git a/Engine/Extras/P4VUtils/Commands/PreflightCommand.cs b/Engine/Extras/P4VUtils/Commands/PreflightCommand.cs
--- a/Engine/Extras/P4VUtils/Commands/PreflightCommand.cs
+++ b/Engine/Extras/P4VUtils/Commands/PreflightCommand.cs
@@ -12,6 +12,7 @@ using System.Threading.Tasks;
 using System.Linq;
 using System.IO;
 using System.Text.RegularExpressions;
+using System.Security.Policy;
 
 namespace P4VUtils.Commands
 {
@@ -157,7 +158,7 @@ namespace P4VUtils.Commands
 
 			string Url = GetUrl(Stream.Stream, Change, ConfigValues);
 			Logger.LogInformation("Opening {Url}", Url);
-			OpenUrl(Url);
+			ProcessUtils.OpenInNewProcess(Url);
 
 			return 0;
 		}
@@ -185,21 +186,6 @@ namespace P4VUtils.Commands
 
 			return BaseUrl.TrimEnd('/');
 		}
-		public static void OpenUrl(string Url)
-		{
-			if (RuntimeInformation.IsOSPlatform(OSPlatform.Windows))
-			{
-				Process.Start(new ProcessStartInfo(Url) { UseShellExecute = true });
-			}
-			else if (RuntimeInformation.IsOSPlatform(OSPlatform.Linux))
-			{
-				Process.Start("xdg-open", Url);
-			}
-			else if (RuntimeInformation.IsOSPlatform(OSPlatform.OSX))
-			{
-				Process.Start("open", Url);
-			}
-		}
 	}
 
 	[Command("preflightandsubmit", CommandCategory.Horde, 2)]
@@ -230,7 +216,7 @@ namespace P4VUtils.Commands
 		}
 	}
 
-	[Command("openpreflight", CommandCategory.Horde, 4)]
+	[Command("openpreflight", CommandCategory.Browser, 1)]
 	class OpenPreflightCommand : Command
 	{
 		public override string Description => "If the changelist has been tagged with #preflight, open the preflight Horde page in the browser";
@@ -290,7 +276,7 @@ namespace P4VUtils.Commands
 				string preflightURL = GetUrl(match!.Groups[1].Value, configValues);
 
 				logger.LogInformation("Opening URL '{URL}' in browser", preflightURL);
-				PreflightCommand.OpenUrl(preflightURL);
+				ProcessUtils.OpenInNewProcess(preflightURL);
 			}
 
 			return 0;
diff --git a/Engine/Extras/P4VUtils/Commands/RobomergePreviewCommand.cs b/Engine/Extras/P4VUtils/Commands/RobomergePreviewCommand.cs
--- a/Engine/Extras/P4VUtils/Commands/RobomergePreviewCommand.cs
+++ b/Engine/Extras/P4VUtils/Commands/RobomergePreviewCommand.cs
@@ -68,7 +68,8 @@ namespace P4VUtils.Commands
 
 			string Url = GetUrl(Change, ConfigValues);
 			Logger.LogInformation("Opening {Url}", Url);
-			OpenUrl(Url);
+
+			ProcessUtils.OpenInNewProcess(Url);
 
 			return 0;
 		}
@@ -84,21 +85,5 @@ namespace P4VUtils.Commands
 			return $"{BaseUrl.TrimEnd('/')}/{Change}";
 
 		}
-
-		static void OpenUrl(string Url)
-		{
-			if (RuntimeInformation.IsOSPlatform(OSPlatform.Windows))
-			{
-				Process.Start(new ProcessStartInfo(Url) { UseShellExecute = true });
-			}
-			else if (RuntimeInformation.IsOSPlatform(OSPlatform.Linux))
-			{
-				Process.Start("xdg-open", Url);
-			}
-			else if (RuntimeInformation.IsOSPlatform(OSPlatform.OSX))
-			{
-				Process.Start("open", Url);
-			}
-		}
 	}
 }
diff --git a/Engine/Extras/P4VUtils/P4VUtils.ini b/Engine/Extras/P4VUtils/P4VUtils.ini
--- a/Engine/Extras/P4VUtils/P4VUtils.ini
+++ b/Engine/Extras/P4VUtils/P4VUtils.ini
@@ -1,2 +1,3 @@
 HordeServer=https://my-horde-server/
-RobomergePreviewURL=https://my-robomerge-server/preview/
\ No newline at end of file
+RobomergePreviewURL=https://my-robomerge-server/preview/
+JiraServer=https://my-jira-server/browse/
\ No newline at end of file
diff --git a/Engine/Extras/P4VUtils/ProcessUtils.cs b/Engine/Extras/P4VUtils/ProcessUtils.cs
new file mode 100644
--- /dev/null
+++ b/Engine/Extras/P4VUtils/ProcessUtils.cs
@@ -0,0 +1,33 @@
+// Copyright Epic Games, Inc. All Rights Reserved.
+
+using System.Diagnostics;
+using System.Runtime.InteropServices;
+
+namespace P4VUtils
+{
+	public static class ProcessUtils
+	{
+		/// <summary>
+		/// Opens the given path in a new process. Which process is launched
+		/// will be decided by the OS based on association with the Path itself.
+		/// A .txt file will probably open a text editor, a url will open the
+		/// default browser etc.
+		/// </summary>
+		/// <param name="Path"The path to be opened></param>
+		public static void OpenInNewProcess(string Path)
+		{
+			if (RuntimeInformation.IsOSPlatform(OSPlatform.Windows))
+			{
+				Process.Start(new ProcessStartInfo(Path) { UseShellExecute = true });
+			}
+			else if (RuntimeInformation.IsOSPlatform(OSPlatform.Linux))
+			{
+				Process.Start("xdg-open", Path);
+			}
+			else if (RuntimeInformation.IsOSPlatform(OSPlatform.OSX))
+			{
+				Process.Start("open", Path);
+			}
+		}
+	}
+}
diff --git a/Engine/Extras/P4VUtils/Program.cs b/Engine/Extras/P4VUtils/Program.cs
--- a/Engine/Extras/P4VUtils/Program.cs
+++ b/Engine/Extras/P4VUtils/Program.cs
@@ -468,6 +468,7 @@ namespace P4VUtils
 				InstallCommandsListInFolder("UE Toolbox", true/*AddFolderToContextMenu*/, CommandCategory.Toolbox, Document, ExecutableLocation, Logger);
 				InstallCommandsListInFolder("UE Integrate", true/*AddFolderToContextMenu*/, CommandCategory.Integrate, Document, ExecutableLocation, Logger);
 				InstallCommandsListInFolder("UE Horde", true/*AddFolderToContextMenu*/, CommandCategory.Horde, Document, ExecutableLocation, Logger);
+				InstallCommandsListInFolder("UE Browser", true/*AddFolderToContextMenu*/, CommandCategory.Browser, Document, ExecutableLocation, Logger);
 			}
 
 			// Save the new document
