Adding new P4VUtil "Combine Shelves", which creates a new shelf combining multiple selected CLs.

diff --git a/Engine/Extras/P4VUtils/Commands/CombineShelves.cs b/Engine/Extras/P4VUtils/Commands/CombineShelves.cs
new file mode 100644
--- /dev/null
+++ b/Engine/Extras/P4VUtils/Commands/CombineShelves.cs
@@ -0,0 +1,122 @@
+﻿// Copyright Epic Games, Inc. All Rights Reserved.
+
+using EpicGames.Core;
+using EpicGames.Perforce;
+using Microsoft.Extensions.Logging;
+using System;
+using System.Collections.Generic;
+using System.ComponentModel;
+using System.Linq;
+using System.Text;
+using System.Threading;
+using System.Threading.Tasks;
+using System.Windows;
+
+namespace P4VUtils.Commands
+{
+
+	[Command("CombineShelves", CommandCategory.Toolbox)]
+	class CombineShelves : UnshelveCommandBase
+	{
+		public override string Description => "Combine multiple selected shelved changelists into a new separate one. The selected changelists must not have any local changes.";
+
+		public override CustomToolInfo CustomTool => new CustomToolInfo("Combine Shelves", "%C")
+		{
+			ShowConsole = true
+		};
+
+		public override async Task<int> Execute(string[] args, IReadOnlyDictionary<string, string> configValues,
+			ILogger logger)
+		{
+			// called when your command is executed, return a value of 0 on success
+			// first argument is the name of the command being executed
+			// second argument is the CL description
+			// subsequent arguments are the individual CLs
+
+			if (args.Length < 2)
+			{
+				logger.LogError("Please select at least one changelist for the operation.");
+				return 1;
+			}
+
+			int temp = 0;
+			string tempString;
+			// validate incoming arguments
+			List<int> changelists = args[1..]
+									.Where(ts => int.TryParse(ts, out temp))
+									.Select(ts => temp).ToList();
+
+			List<string> invalidArguments = args[1..]
+				.Where(x => { return !int.TryParse(x, out temp); }).Select(x => x).ToList();
+
+			// if any of the input changelists couldn't be parsed to integers, 
+			if (invalidArguments.Count > 0)
+			{
+				logger.LogError("Some Input changelist arguments could not be parsed to integers:\n\t{InvalidCLs}",
+					string.Join("\n\t", invalidArguments));
+				return 1;
+			}
+			
+			using PerforceConnection Perforce = new PerforceConnection(null, null, logger);
+
+			bool anyErrors = false;
+
+			// validate the changelists we're about to work with
+			foreach (int changelist in changelists)
+			{
+				ChangeRecord record = await Perforce.GetChangeAsync(GetChangeOptions.None, changelist);
+				if (record.Files.Count != 0)
+				{
+					logger.LogError(
+					"CL {PendingFilesChangelist} has pending files, please shelve or revert files before proceeding.",
+						changelist);
+					anyErrors = true;
+				}
+
+				if (record.Status == ChangeStatus.Submitted)
+				{
+					logger.LogError("CL {SubmittedChangelist} has already been submitted.", changelist);
+					anyErrors = true;
+				}
+			}
+
+			if (anyErrors)
+			{
+				logger.LogError("Changelist validation failed, cancelling operation.");
+				return 1;
+			}
+
+			// make a new changelist to which to shelve our files
+			string description = $"Combined shelf of CLs {string.Join(",", changelists)}";
+			ChangeRecord targetChangelist = new ChangeRecord{ Description = description };
+			
+			logger.LogInformation("Creating new changelist for operation...");
+			
+			await Perforce.CreateChangeAsync(targetChangelist);
+			
+			logger.LogInformation("New Changelist number is {NewCL}", targetChangelist.Number);
+
+			foreach (int changelist in changelists)
+			{
+				logger.LogInformation("Unshelving CL {Changelist} to new CL {NewChangelist}", 
+					changelist,
+					targetChangelist.Number);
+
+				await Perforce.UnshelveAsync(changelist, targetChangelist.Number, null, null, null,
+					UnshelveOptions.ForceOverwrite, new[] {"//..."}, CancellationToken.None);
+			}
+			
+			logger.LogInformation("Shelving all files in CL {NewChangelist}", targetChangelist.Number);
+
+			await Perforce.ShelveAsync(targetChangelist.Number, ShelveOptions.Overwrite, new[] {"//..."});
+
+			logger.LogInformation("Reverting local files in CL {NewChangelist}", targetChangelist.Number);
+			await Perforce.RevertAsync(targetChangelist.Number, null,
+				RevertOptions.DeleteAddedFiles, new[] {"//..."});
+
+			logger.LogInformation("Success!");
+
+			return 0;
+		}
+	}
+}
\ No newline at end of file
