#P4VUtils ConvertToEdit will now clear the readonly flag from files marked for add.
This is necessary because reverting a branch leaves a file with the flag on, and p4 add does not clear the flag. p4 edit does not have this problem because p4 edit clears the flag automatically.

diff --git a/Engine/Extras/P4VUtils/Commands/ConvertToEditCommand.cs b/Engine/Extras/P4VUtils/Commands/ConvertToEditCommand.cs
--- a/Engine/Extras/P4VUtils/Commands/ConvertToEditCommand.cs
+++ b/Engine/Extras/P4VUtils/Commands/ConvertToEditCommand.cs
@@ -6,6 +6,7 @@ using Microsoft.Extensions.Logging;
 using System;
 using System.Collections.Generic;
 using System.ComponentModel;
+using System.IO;
 using System.Linq;
 using System.Text;
 using System.Threading;
@@ -148,6 +149,27 @@ namespace P4VUtils.Commands
 				}
 			}
 
+			// Clear the readonly flag from adds. Edits do not need to do this because p4 edit will do it automatically
+			if (AddActions.Count > 0)
+			{
+				Logger.LogInformation("Clearing the read-only flag from adds");
+
+				string[] LocalAddFiles = AddActions.Select(x => PerforceUtils.EscapePath(OpenedClientToLocalMap[x.ClientFile!])).ToArray();
+				foreach (string LocalFile in LocalAddFiles)
+				{
+					Logger.LogInformation("    {LocalFile}", LocalFile);
+				}
+
+				if (!Debug)
+				{
+					foreach (string LocalAddFile in LocalAddFiles)
+					{
+						FileInfo FileInfo = new FileInfo(LocalAddFile);
+						FileInfo.IsReadOnly = false;
+					}
+				}
+			}
+
 			// Perform basic actions
 			await ReopenFiles("delete", DeleteActions, OpenedClientToLocalMap, x => Perforce.DeleteAsync(Change, DeleteOptions.None, x, CancellationToken.None), Debug, Logger);
 			await ReopenFiles("add", AddActions, OpenedClientToLocalMap, x => Perforce.AddAsync(Change, x, CancellationToken.None), Debug, Logger);
