P4VUtils: Change the Compile Command to use RunUBT.bat

Ensures UBT is compiled and not out of date.

#rnx

diff --git a/Engine/Extras/P4VUtils/Commands/CompileCommand.cs b/Engine/Extras/P4VUtils/Commands/CompileCommand.cs
--- a/Engine/Extras/P4VUtils/Commands/CompileCommand.cs
+++ b/Engine/Extras/P4VUtils/Commands/CompileCommand.cs
@@ -84,7 +84,7 @@ namespace P4VUtils.Commands
 				{
 					ProjectFiles.UnionWith(DirectoryReference.EnumerateFiles(CurrentDir, "*.uproject"));
 
-					FileReference BatchFile = FileReference.Combine(CurrentDir, "Engine/Build/BatchFiles/Build.bat");
+					FileReference BatchFile = FileReference.Combine(CurrentDir, "Engine/Build/BatchFiles/RunUBT.bat");
 					if (FileReference.Exists(BatchFile))
 					{
 						RootDirs.Add(CurrentDir);
@@ -137,7 +137,7 @@ namespace P4VUtils.Commands
 			FileReference TargetListFile = FileReference.Combine(EngineIntermediateDir, "P4VUtils-Targets.txt");
 			WriteLines(TargetListFile, Targets);
 
-			FileReference BuildBatchFile = FileReference.Combine(RootDir, @"Engine\Build\BatchFiles\Build.bat");
+			FileReference BuildBatchFile = FileReference.Combine(RootDir, @"Engine\Build\BatchFiles\RunUBT.bat");
 			StringBuilder Arguments = new StringBuilder($"{BuildBatchFile.FullName.QuoteArgument()} -TargetList={TargetListFile.FullName.QuoteArgument()}");
 			foreach (string AdditionalArgument in AdditionalArguments)
 			{
