P4VUtils: PreflightCommand Refactor

Abstracting out parts to let it be used as a base class

#rnx

diff --git a/Engine/Extras/P4VUtils/Commands/PreflightCommand.cs b/Engine/Extras/P4VUtils/Commands/PreflightCommand.cs
--- a/Engine/Extras/P4VUtils/Commands/PreflightCommand.cs
+++ b/Engine/Extras/P4VUtils/Commands/PreflightCommand.cs
@@ -1,35 +1,56 @@
 // Copyright Epic Games, Inc. All Rights Reserved.
 
-using EpicGames.Core;
 using EpicGames.Perforce;
 using Microsoft.Extensions.Logging;
 using System;
 using System.Collections.Generic;
-using System.Diagnostics;
-using System.Runtime.InteropServices;
-using System.Threading;
-using System.Threading.Tasks;
 using System.Linq;
-using System.IO;
 using System.Text.RegularExpressions;
-using System.Security.Policy;
+using System.Threading;
+using System.Threading.Tasks;
 
 namespace P4VUtils.Commands
 {
 	[Command("preflight", CommandCategory.Horde, 1)]
-	class PreflightCommand : Command
+	class PreflightCommand : Command, IDisposable
 	{
 		static public string StripReviewFyiHashTags(string InString)
 		{
 			return InString.Replace("#review", "-review", StringComparison.Ordinal).Replace("#codereview", "-codereview", StringComparison.Ordinal).Replace("#fyi", "-fyi", StringComparison.Ordinal);
 		}
+
+		internal int Change;
+		internal string? ClientName;
+		internal PerforceConnection? Perforce;
+		internal ClientRecord? Client;
+		internal StreamRecord? Stream;
+		internal DescribeRecord? DescribeRecord;
+		internal List<OpenedRecord>? OpenedRecords;
+
 		public override string Description => "Runs a preflight of the given changelist on Horde";
 
 		public override CustomToolInfo CustomTool => new CustomToolInfo("Horde: Preflight...", "%p");
 
 		public override async Task<int> Execute(string[] Args, IReadOnlyDictionary<string, string> ConfigValues, ILogger Logger)
 		{
-			int Change;
+			if (!await ParseArguments(Args, ConfigValues, Logger))
+			{
+				return 1;
+			}
+
+			if (!await PrepareChangelist(Logger))
+			{
+				// User opted out, a clean exit.
+				return 0;
+			}
+
+			GenerateAndOpenUrl(ConfigValues, Logger);
+
+			return 0;
+		}
+
+		internal virtual async Task<bool> ParseArguments(string[] Args, IReadOnlyDictionary<string, string> ConfigValues, ILogger Logger)
+		{
 			if (Args.Length < 2)
 			{
 				Logger.LogError("Missing changelist number");
@@ -40,7 +61,7 @@ namespace P4VUtils.Commands
 
 					"Invalid Tool Installation?",
 					Logger);
-				return 1;
+				return false;
 			}
 			else if (!int.TryParse(Args[1], out Change))
 			{
@@ -53,55 +74,62 @@ namespace P4VUtils.Commands
 
 					"This changelist requires manual fixes",
 					Logger);
-				return 1;
+				return false;
 			}
 
-			string? ClientName = Environment.GetEnvironmentVariable("P4CLIENT");
-			using PerforceConnection Perforce = new PerforceConnection(null, null, ClientName, Logger);
+			ClientName = Environment.GetEnvironmentVariable("P4CLIENT");
+			Perforce = new PerforceConnection(null, null, ClientName, Logger);
 
-			ClientRecord Client = await Perforce.GetClientAsync(ClientName, CancellationToken.None);
-			if(Client.Stream == null)
+			Client = await Perforce.GetClientAsync(ClientName, CancellationToken.None);
+			if (Client.Stream == null)
 			{
 				Logger.LogError("Not being run from a stream client");
-				return 1;
+				return false;
 			}
 
-			List<OpenedRecord> OpenedRecords = await Perforce.OpenedAsync(OpenedOptions.None, Change, null, null, -1, FileSpecList.Any, CancellationToken.None).ToListAsync();
+			OpenedRecords = await Perforce.OpenedAsync(OpenedOptions.None, Change, null, null, -1, FileSpecList.Any, CancellationToken.None).ToListAsync();
 
 			if (OpenedRecords.Count > 0)
 			{
 				Logger.LogInformation("Shelving changelist {Change}", Change);
 				await Perforce.ShelveAsync(Change, ShelveOptions.Overwrite, new[] { "//..." }, CancellationToken.None);
 			}
-				
+
 			List<DescribeRecord> Describe = await Perforce.DescribeAsync(DescribeOptions.Shelved, -1, new[] { Change }, CancellationToken.None);
 			if (Describe[0].Files.Count == 0)
 			{
 				Logger.LogError("No files are shelved in the given changelist");
-				return 1;
+				return false;
 			}
 
-			StreamRecord Stream = await Perforce.GetStreamAsync(Client.Stream, false, CancellationToken.None);
+			DescribeRecord = Describe[0];
+
+			Stream = await Perforce.GetStreamAsync(Client.Stream, false, CancellationToken.None);
 			while (Stream.Type == "virtual" && Stream.Parent != null)
 			{
 				Stream = await Perforce.GetStreamAsync(Stream.Parent, false, CancellationToken.None);
 			}
 
+			return true;
+		}
+
+		internal virtual async Task<bool> PrepareChangelist(ILogger Logger)
+		{
 			if (CreateBackupCL())
 			{
 				// if this CL has files still open within it, create a new CL for those still opened files
 				// before sending the original CL to the preflight system - this avoids the problem where Horde
 				// cannot take ownership of the original CL and fails to checkin
 
-				if (OpenedRecords.Count > 0)
+				if (OpenedRecords!.Count > 0)
 				{
-					InfoRecord Info = await Perforce.GetInfoAsync(InfoOptions.None, CancellationToken.None);
+					InfoRecord Info = await Perforce!.GetInfoAsync(InfoOptions.None, CancellationToken.None);
 
 					ChangeRecord NewChangeRecord = new ChangeRecord();
 					NewChangeRecord.User = Info.UserName;
 					NewChangeRecord.Client = Info.ClientName;
-					NewChangeRecord.Description = $"{StripReviewFyiHashTags(Describe[0].Description.TrimEnd())}\n#p4v-preflight-copy {Change}";
-					NewChangeRecord = await Perforce.CreateChangeAsync(NewChangeRecord, CancellationToken.None);
+					NewChangeRecord.Description = $"{StripReviewFyiHashTags(DescribeRecord!.Description.TrimEnd())}\n#p4v-preflight-copy {Change}";
+					NewChangeRecord = await Perforce!.CreateChangeAsync(NewChangeRecord, CancellationToken.None);
 
 					Logger.LogInformation("Created pending changelist {Change}", NewChangeRecord.Number);
 
@@ -109,7 +137,7 @@ namespace P4VUtils.Commands
 					{
 						if (OpenedRecord.ClientFile != null)
 						{
-							await Perforce.ReopenAsync(NewChangeRecord.Number, OpenedRecord.Type, OpenedRecord.ClientFile!, CancellationToken.None);
+							await Perforce!.ReopenAsync(NewChangeRecord.Number, OpenedRecord.Type, OpenedRecord.ClientFile!, CancellationToken.None);
 							Logger.LogInformation("moving opened {File} to CL {CL}", OpenedRecord.ClientFile.ToString(), NewChangeRecord.Number);
 						}
 					}
@@ -119,48 +147,47 @@ namespace P4VUtils.Commands
 					Logger.LogInformation("No files opened, no copy CL created");
 				}
 			}
-			else
-			{
-				// if this CL has files still open within it, and this is a submit request - warn the user and provide options
+			// if this CL has files still open within it, and this is a submit request - warn the user and provide options
 
-				if (OpenedRecords.Count > 0 && IsSubmit())
-				{
-					string Prompt = "Your CL was just shelved however it still has files checked out\r\n" +
-							"If the files remain in the CL your preflight will fail to submit\r\n" +
-							"\r\n" +
-							"Click:\r\n" +
-							"[YES] - To revert local files and submit the preflight,\r\n" +
-							"[NO] - To start the preflight, and move the files manually\r\n" +
-							"[CANCEL] - To cancel the request";
-					string Caption = "Your CL will fail to auto-submit unless fixed";
+			if (OpenedRecords!.Count > 0 && IsSubmit())
+			{
+				string Prompt = "Your CL was just shelved however it still has files checked out\r\n" +
+						"If the files remain in the CL your preflight will fail to submit\r\n" +
+						"\r\n" +
+						"Click:\r\n" +
+						"[YES] - To revert local files and submit the preflight,\r\n" +
+						"[NO] - To start the preflight, and move the files manually\r\n" +
+						"[CANCEL] - To cancel the request";
+				string Caption = "Your CL will fail to auto-submit unless fixed";
 
-					UserInterface.Button Response = UserInterface.ShowDialog(Prompt, Caption, UserInterface.YesNoCancel, UserInterface.Button.Yes, Logger);
+				UserInterface.Button Response = UserInterface.ShowDialog(Prompt, Caption, UserInterface.YesNoCancel, UserInterface.Button.Yes, Logger);
 
 
 
-					if (Response == UserInterface.Button.No)
-					{
-						// do nothing - user has been warned.
-					}
-					else if (Response == UserInterface.Button.Yes)
-					{
-						await Perforce.RevertAsync(Change, null, RevertOptions.None, OpenedRecords.Select(x => x.ClientFile!).ToArray(), CancellationToken.None);
-					}
-					// any other reply is Cancel (like on Mac, hitting Escape will return a weird string, not Cancel)
-					else
-					{
-						Logger.LogInformation("User Opted to cancel");
-						return 0;
-					}
+				if (Response == UserInterface.Button.No)
+				{
+					// do nothing - user has been warned.
+				}
+				else if (Response == UserInterface.Button.Yes)
+				{
+					await Perforce!.RevertAsync(Change, null, RevertOptions.None, OpenedRecords.Select(x => x.ClientFile!).ToArray(), CancellationToken.None);
+				}
+				// any other reply is Cancel (like on Mac, hitting Escape will return a weird string, not Cancel)
+				else
+				{
+					Logger.LogInformation("User Opted to cancel");
+					return false;
 				}
 			}
 
+			return true;
+		}
 
-			string Url = GetUrl(Stream.Stream, Change, ConfigValues);
+		internal virtual void GenerateAndOpenUrl(IReadOnlyDictionary<string, string> ConfigValues, ILogger Logger)
+		{
+			string Url = GetUrl(Stream!.Stream, Change, ConfigValues);
 			Logger.LogInformation("Opening {Url}", Url);
 			ProcessUtils.OpenInNewProcess(Url);
-
-			return 0;
 		}
 
 		public virtual bool CreateBackupCL()
@@ -186,6 +213,11 @@ namespace P4VUtils.Commands
 
 			return BaseUrl.TrimEnd('/');
 		}
+
+		public void Dispose()
+		{
+			Perforce?.Dispose();
+		}
 	}
 
 	[Command("preflightandsubmit", CommandCategory.Horde, 2)]
