P4VUtils: Fixing moving files command

#rnx

diff --git a/Engine/Extras/P4VUtils/Commands/PreflightCommand.cs b/Engine/Extras/P4VUtils/Commands/PreflightCommand.cs
--- a/Engine/Extras/P4VUtils/Commands/PreflightCommand.cs
+++ b/Engine/Extras/P4VUtils/Commands/PreflightCommand.cs
@@ -147,36 +147,39 @@ namespace P4VUtils.Commands
 					Logger.LogInformation("No files opened, no copy CL created");
 				}
 			}
-			// if this CL has files still open within it, and this is a submit request - warn the user and provide options
-
-			if (OpenedRecords!.Count > 0 && IsSubmit())
+			else
 			{
-				string Prompt = "Your CL was just shelved however it still has files checked out\r\n" +
-						"If the files remain in the CL your preflight will fail to submit\r\n" +
-						"\r\n" +
-						"Click:\r\n" +
-						"[YES] - To revert local files and submit the preflight,\r\n" +
-						"[NO] - To start the preflight, and move the files manually\r\n" +
-						"[CANCEL] - To cancel the request";
-				string Caption = "Your CL will fail to auto-submit unless fixed";
+				// if this CL has files still open within it, and this is a submit request - warn the user and provide options
 
-				UserInterface.Button Response = UserInterface.ShowDialog(Prompt, Caption, UserInterface.YesNoCancel, UserInterface.Button.Yes, Logger);
+				if (OpenedRecords!.Count > 0 && IsSubmit())
+				{
+					string Prompt = "Your CL was just shelved however it still has files checked out\r\n" +
+							"If the files remain in the CL your preflight will fail to submit\r\n" +
+							"\r\n" +
+							"Click:\r\n" +
+							"[YES] - To revert local files and submit the preflight,\r\n" +
+							"[NO] - To start the preflight, and move the files manually\r\n" +
+							"[CANCEL] - To cancel the request";
+					string Caption = "Your CL will fail to auto-submit unless fixed";
 
+					UserInterface.Button Response = UserInterface.ShowDialog(Prompt, Caption, UserInterface.YesNoCancel, UserInterface.Button.Yes, Logger);
 
 
-				if (Response == UserInterface.Button.No)
-				{
-					// do nothing - user has been warned.
-				}
-				else if (Response == UserInterface.Button.Yes)
-				{
-					await Perforce!.RevertAsync(Change, null, RevertOptions.None, OpenedRecords.Select(x => x.ClientFile!).ToArray(), CancellationToken.None);
-				}
-				// any other reply is Cancel (like on Mac, hitting Escape will return a weird string, not Cancel)
-				else
-				{
-					Logger.LogInformation("User Opted to cancel");
-					return false;
+
+					if (Response == UserInterface.Button.No)
+					{
+						// do nothing - user has been warned.
+					}
+					else if (Response == UserInterface.Button.Yes)
+					{
+						await Perforce!.RevertAsync(Change, null, RevertOptions.None, OpenedRecords.Select(x => x.ClientFile!).ToArray(), CancellationToken.None);
+					}
+					// any other reply is Cancel (like on Mac, hitting Escape will return a weird string, not Cancel)
+					else
+					{
+						Logger.LogInformation("User Opted to cancel");
+						return false;
+					}
 				}
 			}
 
