P4VUtils: Fixing build warning about unused var

#rnx

diff --git a/Engine/Extras/P4VUtils/Commands/CombineShelves.cs b/Engine/Extras/P4VUtils/Commands/CombineShelves.cs
--- a/Engine/Extras/P4VUtils/Commands/CombineShelves.cs
+++ b/Engine/Extras/P4VUtils/Commands/CombineShelves.cs
@@ -40,7 +40,6 @@ namespace P4VUtils.Commands
 			}
 
 			int temp = 0;
-			string tempString;
 			// validate incoming arguments
 			List<int> changelists = args[1..]
 									.Where(ts => int.TryParse(ts, out temp))
