[P4VUtils] Correctly return error code from the Compile Command (needed for the submit tool validation)

#rnx

diff --git a/Engine/Extras/P4VUtils/Commands/CompileCommand.cs b/Engine/Extras/P4VUtils/Commands/CompileCommand.cs
--- a/Engine/Extras/P4VUtils/Commands/CompileCommand.cs
+++ b/Engine/Extras/P4VUtils/Commands/CompileCommand.cs
@@ -155,6 +155,11 @@ namespace P4VUtils.Commands
 			{
 #pragma warning disable CA2000 // Dispose objects before losing scope
 				await Process.CopyToAsync(Console.OpenStandardOutput(), CancellationToken.None);
+				
+				if (Process.ExitCode != 0)
+				{
+					return false;
+				}
 #pragma warning restore CA2000 // Dispose objects before losing scope
 			}
 
