[P4VUtils] Add -Define:CHECK_PUREVIRTUALS=1 to the compile command

#rnx

diff --git a/Engine/Extras/P4VUtils/Commands/CompileCommand.cs b/Engine/Extras/P4VUtils/Commands/CompileCommand.cs
--- a/Engine/Extras/P4VUtils/Commands/CompileCommand.cs
+++ b/Engine/Extras/P4VUtils/Commands/CompileCommand.cs
@@ -144,6 +144,9 @@ namespace P4VUtils.Commands
 				Arguments.AppendFormat(" {0}", AdditionalArgument.QuoteArgument());
 			}
 
+			// Enable checking of pure virtuals
+			Arguments.Append(" -Define:CHECK_PUREVIRTUALS=1");
+
 			Logger.LogInformation("Running {Arguments}", Arguments);
 			Logger.LogInformation("");
 
