P4VUtils: New Commands

* Robomerge Track Change - track where a submitted change has merged to
* Horde Validate Configs - validate Horde server config files

#rnx

diff --git a/Engine/Extras/P4VUtils/Commands/PreflightCommand.cs b/Engine/Extras/P4VUtils/Commands/PreflightCommand.cs
--- a/Engine/Extras/P4VUtils/Commands/PreflightCommand.cs
+++ b/Engine/Extras/P4VUtils/Commands/PreflightCommand.cs
@@ -4,6 +4,7 @@ using EpicGames.Perforce;
 using Microsoft.Extensions.Logging;
 using System;
 using System.Collections.Generic;
+using System.IO;
 using System.Linq;
 using System.Text.RegularExpressions;
 using System.Threading;
@@ -323,4 +324,51 @@ namespace P4VUtils.Commands
 			return $"{BaseUrl}/job/{preflightId}";
 		}
 	}
+
+	[Command("preflighthordeconfig", CommandCategory.Horde, 4)]
+	class PreflightHordeConfigCommand : PreflightCommand
+	{
+		public override string Description => "If the changelist contains Horde configuration file(s), open Horde in the browser to validate the configuration file(s)";
+
+		public override CustomToolInfo CustomTool => new CustomToolInfo("Horde: Validate Configuration Files...", "%p");
+
+		/// <summary>
+		/// Is it a Horde server configuration file, ie: ends with .stream.json, .project.json, or is named globals.json
+		/// </summary>
+		/// <param name="DepotPath"></param>
+		/// <returns></returns>
+		static internal bool IsHordeConfigurationFile(string DepotPath)
+		{
+			string FileName = Path.GetFileName(DepotPath);
+			return (FileName.EndsWith(".stream.json", StringComparison.OrdinalIgnoreCase) || 
+					FileName.EndsWith(".project.json", StringComparison.OrdinalIgnoreCase) || 
+					FileName.Equals("globals.json", StringComparison.OrdinalIgnoreCase));
+		}
+
+		internal override async Task<bool> ParseArguments(string[] Args, IReadOnlyDictionary<string, string> ConfigValues, ILogger Logger)
+		{
+			if (!await base.ParseArguments(Args, ConfigValues, Logger))
+			{
+				return false;
+			}
+
+			if (!DescribeRecord!.Files.Any(x => IsHordeConfigurationFile(x.DepotFile)))
+			{
+				Logger.LogError("No Horde Configuration Files");
+				Logger.LogError("'{Change}' does not contain Horde Configuration files.", Change);
+				UserInterface.ShowSimpleDialog(
+					$"The specified changelist, {Change}, does not contain any Horde Configuration files.",
+					"Invalid Changelist", Logger);
+				return false;
+			}
+
+			return true;
+		}
+
+		public override string GetUrl(string Stream, int Change, IReadOnlyDictionary<string, string> ConfigValues)
+		{
+			string BaseUrl = PreflightCommand.GetHordeServerAddress(ConfigValues);
+			return $"{BaseUrl}/preflightconfig?shelvedchange={Change}";
+		}
+	}
 }
diff --git a/Engine/Extras/P4VUtils/Commands/RobomergePreviewCommand.cs b/Engine/Extras/P4VUtils/Commands/RobomergePreviewCommand.cs
--- a/Engine/Extras/P4VUtils/Commands/RobomergePreviewCommand.cs
+++ b/Engine/Extras/P4VUtils/Commands/RobomergePreviewCommand.cs
@@ -4,9 +4,7 @@ using EpicGames.Perforce;
 using Microsoft.Extensions.Logging;
 using System;
 using System.Collections.Generic;
-using System.Diagnostics;
 using System.Linq;
-using System.Runtime.InteropServices;
 using System.Threading;
 using System.Threading.Tasks;
 
@@ -74,15 +72,59 @@ namespace P4VUtils.Commands
 			return 0;
 		}
 
-		public virtual string GetUrl(int Change, IReadOnlyDictionary<string, string> ConfigValues)
+		static public string GetUrl(int Change, IReadOnlyDictionary<string, string> ConfigValues)
 		{
 			string? BaseUrl;
-			if (!ConfigValues.TryGetValue("RobomergePreviewURL", out BaseUrl))
+			if (!ConfigValues.TryGetValue("RobomergeURL", out BaseUrl))
 			{
 				BaseUrl = "https://configure-server-url-in-p4vutils.ini";
 			}
 
-			return $"{BaseUrl.TrimEnd('/')}/{Change}";
+			return $"{BaseUrl.TrimEnd('/')}/preview/{Change}";
+
+		}
+	}
+
+	[Command("robomergetrackchange", CommandCategory.Toolbox, 11)]
+	class RobomergeTrackChangeCommand : Command
+	{
+		public override string Description => "Opens the RoboMerge site to track where a submitted change has been merged to.";
+
+		public override CustomToolInfo CustomTool => new CustomToolInfo("Robomerge Track Change", "%c") { ShowConsole = true, RefreshUI = false };
+
+		public override async Task<int> Execute(string[] Args, IReadOnlyDictionary<string, string> ConfigValues, ILogger Logger)
+		{
+			await Task.Yield();
+
+			int Change;
+			if (Args.Length < 2)
+			{
+				Logger.LogError("Missing changelist number");
+				return 1;
+			}
+			else if (!int.TryParse(Args[1], out Change))
+			{
+				Logger.LogError("'{Argument}' is not a numbered changelist", Args[1]);
+				return 1;
+			}
+
+			string Url = GetUrl(Change, ConfigValues);
+			Logger.LogInformation("Opening {Url}", Url);
+
+			ProcessUtils.OpenInNewProcess(Url);
+
+			return 0;
+		}
+
+		static public string GetUrl(int Change, IReadOnlyDictionary<string, string> ConfigValues)
+		{
+			string? BaseUrl;
+			if (!ConfigValues.TryGetValue("RobomergeURL", out BaseUrl))
+			{
+				BaseUrl = "https://configure-server-url-in-p4vutils.ini";
+			}
+
+			return $"{BaseUrl.TrimEnd('/')}/trackchange/{Change}";
 
 		}
 	}
diff --git a/Engine/Extras/P4VUtils/P4VUtils.ini b/Engine/Extras/P4VUtils/P4VUtils.ini
--- a/Engine/Extras/P4VUtils/P4VUtils.ini
+++ b/Engine/Extras/P4VUtils/P4VUtils.ini
@@ -1,3 +1,3 @@
 HordeServer=https://my-horde-server/
-RobomergePreviewURL=https://my-robomerge-server/preview/
+RobomergeURL=https://my-robomerge-server/
 JiraServer=https://my-jira-server/browse/
\ No newline at end of file
