Add a convert move to add/delete command
#rnx

diff --git a/Engine/Extras/P4VUtils/Commands/ConvertMoveToAddDeleteCommands.cs b/Engine/Extras/P4VUtils/Commands/ConvertMoveToAddDeleteCommands.cs
new file mode 100644
--- /dev/null
+++ b/Engine/Extras/P4VUtils/Commands/ConvertMoveToAddDeleteCommands.cs
@@ -0,0 +1,164 @@
+// Copyright Epic Games, Inc. All Rights Reserved.
+
+using EpicGames.Core;
+using EpicGames.Perforce;
+using Microsoft.Extensions.Logging;
+using System;
+using System.Collections.Generic;
+using System.ComponentModel;
+using System.IO;
+using System.Linq;
+using System.Text;
+using System.Threading;
+using System.Threading.Tasks;
+
+namespace P4VUtils.Commands
+{
+	[Command("convertmovetoadddelete", CommandCategory.Integrate, 1)]
+	class ConvertMoveToAddDeleteCommand : Command
+	{
+		public override string Description => "Converts move operations to add/deletes";
+
+		public override CustomToolInfo CustomTool => new CustomToolInfo("Convert Move to Add/Delete", "%p") { ShowConsole = true };
+
+		public override async Task<int> Execute(string[] Args, IReadOnlyDictionary<string, string> ConfigValues, ILogger Logger)
+		{
+			int Change;
+			if (Args.Length < 2)
+			{
+				Logger.LogError("Missing changelist number");
+				return 1;
+			}
+			else if (!int.TryParse(Args[1], out Change))
+			{
+				Logger.LogError("'{Argument}' is not a numbered changelist", Args[1]);
+				return 1;
+			}
+
+			bool Debug = Args.Any(x => x.Equals("-Debug", StringComparison.OrdinalIgnoreCase));
+
+			using PerforceConnection Perforce = new PerforceConnection(null, null, Logger);
+			if(!await ConvertMoveToAddDeleteAsync(Perforce, Change, Debug, Logger))
+			{
+				return 1;
+			}
+			return 0;
+		}
+
+		public static async Task<bool> ConvertMoveToAddDeleteAsync(PerforceConnection Perforce, int Change, bool Debug, ILogger Logger)
+		{
+			// Get the list of opened files in the CL along with their metadata
+			Logger.LogInformation("Getting the list of files in changelist {Change}...", Change);
+			List<OpenedRecord> OpenedRecords = await Perforce.OpenedAsync(OpenedOptions.None, Change, null, null, -1, FileSpecList.Any, CancellationToken.None).ToListAsync();
+			if (OpenedRecords.Count == 0)
+			{
+				Logger.LogError("Change has no opened files. Aborting.");
+				return false;
+			}
+
+			// Each record has {depotFile, clientFile, rev, headRev, action, change, type, user, client}
+			Logger.LogInformation("There are {NumFiles} open files in the change.", OpenedRecords.Count);
+
+			// create a mapping of client to local files so we can re-open using local paths (needed for some actions)
+			List<WhereRecord> WhereRecords = await Perforce.WhereAsync(OpenedRecords.Select(x => x.DepotFile!).ToArray(), CancellationToken.None).ToListAsync();
+			Dictionary<string, string> OpenedClientToLocalMap = WhereRecords.ToDictionary(x => x.ClientFile, x => x.Path, StringComparer.OrdinalIgnoreCase);
+
+			// If any file that needs to be resolved still, we need to abort.
+			// We don't want to convert non-resolved files to plain edits,
+			// as we need to resolve the merge first!
+			List<ResolveRecord> ResolveRecords = await Perforce.ResolveAsync(Change, ResolveOptions.PreviewOnly, FileSpecList.Any, CancellationToken.None);
+			if (ResolveRecords.Any())
+			{
+				Logger.LogError("The following files in the changelist need to be resolved.");
+				Logger.LogError("This script cannot continue until all files are resolved:");
+				foreach (ResolveRecord ResolveRecord in ResolveRecords)
+				{
+					Logger.LogError("    {FileToResolve}", ResolveRecord.ClientFile);
+				}
+				return false;
+			}
+
+			// We'll just assume all the files in the changelist are integrated. We'll revert and re-open
+			// with the same action/filetype. This should remove the integration records.
+
+			// partition our files into ones we know the actions for, and the ones we don't
+			List<OpenedRecord> MoveAddActions = new List<OpenedRecord>();
+			List<OpenedRecord> MoveDeleteActions = new List<OpenedRecord>();
+			List<OpenedRecord> ToRevertRecords = new List<OpenedRecord>();
+
+			// Multiple action map to the same reopen mapping, so this handles that.
+			Dictionary<FileAction, List<OpenedRecord>> ActionMappings = new Dictionary<FileAction, List<OpenedRecord>>()
+			{
+				[FileAction.MoveAdd] = MoveAddActions,
+				[FileAction.MoveDelete] = MoveDeleteActions
+			};
+
+			// this maps all open records to a reopen action.
+			foreach (OpenedRecord OpenedRecord in OpenedRecords)
+			{
+				List<OpenedRecord>? ActionList;
+				if (ActionMappings.TryGetValue(OpenedRecord.Action, out ActionList))
+				{
+					ActionList.Add(OpenedRecord);
+					ToRevertRecords.Add(OpenedRecord);
+				}
+			}
+
+			// We are good to go, so actually revert the files.
+			Logger.LogInformation("Reverting files that will be reopened");
+			if (!Debug)
+			{
+				PerforceResponseList<RevertRecord> Results = await Perforce.TryRevertAsync(-1, null, RevertOptions.KeepWorkspaceFiles, ToRevertRecords.Select(x => x.ClientFile!).ToArray(), CancellationToken.None);
+				Results.EnsureSuccess();
+			}
+
+			// Clear the readonly flag from adds. Edits do not need to do this because p4 edit will do it automatically
+			if (MoveAddActions.Count > 0)
+			{
+				Logger.LogInformation("Clearing the read-only flag from adds");
+
+				string[] LocalAddFiles = MoveAddActions.Select(x => PerforceUtils.EscapePath(OpenedClientToLocalMap[x.ClientFile!])).ToArray();
+				foreach (string LocalFile in LocalAddFiles)
+				{
+					Logger.LogInformation("    {LocalFile}", LocalFile);
+				}
+
+				if (!Debug)
+				{
+					foreach (string LocalAddFile in LocalAddFiles)
+					{
+						FileInfo FileInfo = new FileInfo(LocalAddFile);
+						FileInfo.IsReadOnly = false;
+					}
+				}
+			}
+
+			// Perform basic actions
+			await ReopenFiles("delete", MoveDeleteActions, OpenedClientToLocalMap, x => Perforce.DeleteAsync(Change, DeleteOptions.None, x, CancellationToken.None), Debug, Logger);
+			await ReopenFiles("add", MoveAddActions, OpenedClientToLocalMap, x => Perforce.AddAsync(Change, x, CancellationToken.None), Debug, Logger);
+
+			// Done!
+			Logger.LogInformation("Success");
+			return true;
+		}
+
+		static async Task ReopenFiles(string Operation, List<OpenedRecord> Records, Dictionary<string, string> OpenedClientToLocalMap, Func<string[], Task> ReopenAsync, bool Debug, ILogger Logger)
+		{
+			if (Records.Count > 0)
+			{
+				Logger.LogInformation("Re-opening the following files for {Operation}:", Operation);
+
+				string[] LocalFiles = Records.Select(x => PerforceUtils.EscapePath(OpenedClientToLocalMap[x.ClientFile!])).ToArray();
+				foreach (string LocalFile in LocalFiles)
+				{
+					Logger.LogInformation("    {LocalFile}", LocalFile);
+				}
+
+				if (!Debug)
+				{
+					await ReopenAsync(LocalFiles);
+				}
+			}
+		}
+	}
+}
