P4VUtils: Support building natively on win-arm64.

diff --git a/Engine/Extras/P4VUtils/P4VUtils.csproj b/Engine/Extras/P4VUtils/P4VUtils.csproj
--- a/Engine/Extras/P4VUtils/P4VUtils.csproj
+++ b/Engine/Extras/P4VUtils/P4VUtils.csproj
@@ -33,6 +33,13 @@
         <PublishDir>$(NFLPrefix)Binaries/Win64</PublishDir>
         <RuntimeIdentifier>win-x64</RuntimeIdentifier>
       </PropertyGroup>
+      <Choose>
+        <When Condition="'$([System.Runtime.InteropServices.RuntimeInformation]::OSArchitecture)' == 'Arm64'">
+          <PropertyGroup>
+            <RuntimeIdentifier>win-arm64</RuntimeIdentifier>
+          </PropertyGroup>
+        </When>
+      </Choose>
     </When>
     <Otherwise>
       <PropertyGroup>
