added submit tool tag to safe back out tool

diff --git a/Engine/Extras/P4VUtils/Commands/BackoutCommand.cs b/Engine/Extras/P4VUtils/Commands/BackoutCommand.cs
--- a/Engine/Extras/P4VUtils/Commands/BackoutCommand.cs
+++ b/Engine/Extras/P4VUtils/Commands/BackoutCommand.cs
@@ -211,7 +211,7 @@ namespace P4VUtils.Commands
 			ChangeRecord NewChangeRecord = new ChangeRecord();
 			NewChangeRecord.User = Info.UserName;
 			NewChangeRecord.Client = Info.ClientName;
-			NewChangeRecord.Description = $"[Backout] - CL{Change}\n#fyi {ExistingChangeRecord.User}\nOriginal CL Desc\n-----------------------------------------------------------------\n{ExistingChangeRecord.Description.TrimEnd()}\n";
+			NewChangeRecord.Description = $"[Backout] - CL{Change}\n#fyi {ExistingChangeRecord.User}\n#submittool safebackout\nOriginal CL Desc\n-----------------------------------------------------------------\n{ExistingChangeRecord.Description.TrimEnd()}\n";
 			NewChangeRecord = await Perforce.CreateChangeAsync(NewChangeRecord, CancellationToken.None);
 
 			Logger.LogInformation("Created pending changelist {Change}", NewChangeRecord.Number);
