P4VUtils: Remove Epic specfic messaing

Add a helper text ini option to customize extra backout messaging

#rnx

diff --git a/Engine/Extras/P4VUtils/Commands/BackoutCommand.cs b/Engine/Extras/P4VUtils/Commands/BackoutCommand.cs
--- a/Engine/Extras/P4VUtils/Commands/BackoutCommand.cs
+++ b/Engine/Extras/P4VUtils/Commands/BackoutCommand.cs
@@ -136,22 +136,25 @@ namespace P4VUtils.Commands
 				string ErrorString = String.Join("\r\n", ValidationResult.Errors); 
 				if (ErrorString.Length > MAX_ERROR_LENGTH)
 				{
-					ErrorString = ErrorString.Substring(0, MAX_ERROR_LENGTH);
-					ErrorString += "\r\n (...)";
+					ErrorString = $"{ErrorString.Substring(0, MAX_ERROR_LENGTH)}\r\n (...)";
+				}
+
+				StringBuilder MessageText = new StringBuilder();
+				MessageText.Append("This backout is potentially unsafe:\r\n");
+				MessageText.Append("\r\n");
+				MessageText.Append(ErrorString);
+				MessageText.Append("\r\n\r\n");
+				MessageText.Append("Are you sure you want to proceed with the backout operation?\r\n");
+				MessageText.Append("\r\n");
+
+				if (ConfigValues.TryGetValue("SafeBackoutHelpText", out string? HelpText))
+				{
+					MessageText.Append(HelpText);
+					MessageText.Append("\r\n");
 				}
 
 				// warn user
-				UserInterface.Button result = UserInterface.ShowDialog(
-					"This backout is potentially unsafe:\r\n" +
-					"\r\n" +
-					ErrorString +
-					"\r\n\r\n" +
-					"Are you sure you want to proceed with the backout operation?\r\n" +
-					"\r\n" +
-					"You can tag @p4backouthelp in #ue-build-health or #fn-build-health for assistance.\r\n" +
-					"\r\n",
-					"Unsafe backout detected",
-					UserInterface.YesNo, UserInterface.Button.No, Logger);
+				UserInterface.Button result = UserInterface.ShowDialog(MessageText.ToString(), "Unsafe backout detected", UserInterface.YesNo, UserInterface.Button.No, Logger);
 
 				if (result == UserInterface.Button.No)
 				{
diff --git a/Engine/Extras/P4VUtils/P4VUtils.ini b/Engine/Extras/P4VUtils/P4VUtils.ini
--- a/Engine/Extras/P4VUtils/P4VUtils.ini
+++ b/Engine/Extras/P4VUtils/P4VUtils.ini
@@ -1,3 +1,4 @@
 HordeServer=https://my-horde-server/
 RobomergeURL=https://my-robomerge-server/
-JiraServer=https://my-jira-server/browse/
\ No newline at end of file
+JiraServer=https://my-jira-server/browse/
+SafeBackoutHelpText=
\ No newline at end of file
