P4VUtils - Updating CombineShelves tool to allow combining both shelved and unshelved CLs

diff --git a/Engine/Extras/P4VUtils/Commands/CombineShelves.cs b/Engine/Extras/P4VUtils/Commands/CombineShelves.cs
--- a/Engine/Extras/P4VUtils/Commands/CombineShelves.cs
+++ b/Engine/Extras/P4VUtils/Commands/CombineShelves.cs
@@ -18,13 +18,22 @@ namespace P4VUtils.Commands
 	[Command("CombineShelves", CommandCategory.Toolbox)]
 	class CombineShelves : UnshelveCommandBase
 	{
-		public override string Description => "Combine multiple selected shelved changelists into a new separate one. The selected changelists must not have any local changes.";
+		public override string Description => "Combine multiple selected shelved changelists into a new separate one.";
 
 		public override CustomToolInfo CustomTool => new CustomToolInfo("Combine Shelves", "%C")
 		{
 			ShowConsole = true
 		};
 
+		private async void ReopenFilesToChangelist(ChangeRecord changeRecord, int targetChangelist, PerforceConnection connection, ILogger logger)
+		{
+			foreach (string fileName in changeRecord.Files)
+			{
+				logger.LogInformation("Moving {Filename} to target CL {TargetChangelist}", fileName, targetChangelist);
+				await connection.ReopenAsync(targetChangelist, null, fileName);
+			}
+		}
+		
 		public override async Task<int> Execute(string[] args, IReadOnlyDictionary<string, string> configValues,
 			ILogger logger)
 		{
@@ -60,17 +69,15 @@ namespace P4VUtils.Commands
 
 			bool anyErrors = false;
 
+			// list to store local files for each changelist, if there are any
+			List<ChangeRecord> changelistRecords = new List<ChangeRecord>();
+			
 			// validate the changelists we're about to work with
 			foreach (int changelist in changelists)
 			{
 				ChangeRecord record = await Perforce.GetChangeAsync(GetChangeOptions.None, changelist);
-				if (record.Files.Count != 0)
-				{
-					logger.LogError(
-					"CL {PendingFilesChangelist} has pending files, please shelve or revert files before proceeding.",
-						changelist);
-					anyErrors = true;
-				}
+				
+				changelistRecords.Add(record);
 
 				if (record.Status == ChangeStatus.Submitted)
 				{
@@ -95,20 +102,37 @@ namespace P4VUtils.Commands
 			
 			logger.LogInformation("New Changelist number is {NewCL}", targetChangelist.Number);
 
-			foreach (int changelist in changelists)
+			// then loop over all the changelists and get those files into the new changelist
+			foreach(ChangeRecord changeRecord in changelistRecords)
 			{
-				logger.LogInformation("Unshelving CL {Changelist} to new CL {NewChangelist}", 
-					changelist,
-					targetChangelist.Number);
+				if (changeRecord.Files.Count == 0)
+				{
+					logger.LogInformation("Unshelving CL {Changelist} to new CL {NewChangelist}",
+						changeRecord.Number,
+						targetChangelist.Number);
 
-				await Perforce.UnshelveAsync(changelist, targetChangelist.Number, null, null, null,
-					UnshelveOptions.ForceOverwrite, new[] {"//..."}, CancellationToken.None);
+					await Perforce.UnshelveAsync(changeRecord.Number, targetChangelist.Number, null, null, null,
+						UnshelveOptions.ForceOverwrite, new[] {"//..."}, CancellationToken.None);
+				}
+				else
+				{
+					ReopenFilesToChangelist(changeRecord, targetChangelist.Number, Perforce, logger);
+				}
 			}
 			
 			logger.LogInformation("Shelving all files in CL {NewChangelist}", targetChangelist.Number);
 
 			await Perforce.ShelveAsync(targetChangelist.Number, ShelveOptions.Overwrite, new[] {"//..."});
 
+			// once the new shelf is established, move the originally unshelved files back to their original CLs 
+			foreach (ChangeRecord changeRecord in changelistRecords)
+			{
+				if (changeRecord.Files.Count != 0)
+				{
+					ReopenFilesToChangelist(changeRecord, changeRecord.Number, Perforce, logger);
+				}
+			}
+			
 			logger.LogInformation("Reverting local files in CL {NewChangelist}", targetChangelist.Number);
 			await Perforce.RevertAsync(targetChangelist.Number, null,
 				RevertOptions.DeleteAddedFiles, new[] {"//..."});
