P4VUtils - Marking ReopenFilesToChangelist as static, per CA1822

diff --git a/Engine/Extras/P4VUtils/Commands/CombineShelves.cs b/Engine/Extras/P4VUtils/Commands/CombineShelves.cs
--- a/Engine/Extras/P4VUtils/Commands/CombineShelves.cs
+++ b/Engine/Extras/P4VUtils/Commands/CombineShelves.cs
@@ -25,7 +25,7 @@ namespace P4VUtils.Commands
 			ShowConsole = true
 		};
 
-		private async void ReopenFilesToChangelist(ChangeRecord changeRecord, int targetChangelist, PerforceConnection connection, ILogger logger)
+		private static async void ReopenFilesToChangelist(ChangeRecord changeRecord, int targetChangelist, PerforceConnection connection, ILogger logger)
 		{
 			foreach (string fileName in changeRecord.Files)
 			{
