Fix warning CA1822: Member 'FindLastEditInFileHistory' does not access instance data and can be marked as static
#rnx

diff --git a/Engine/Extras/P4VUtils/Commands/FindLastEditCommand.cs b/Engine/Extras/P4VUtils/Commands/FindLastEditCommand.cs
--- a/Engine/Extras/P4VUtils/Commands/FindLastEditCommand.cs
+++ b/Engine/Extras/P4VUtils/Commands/FindLastEditCommand.cs
@@ -18,7 +18,7 @@ namespace P4VUtils.Commands
 
 		public override CustomToolInfo CustomTool => new CustomToolInfo("Find last edit", "%d") { ShowConsole = true };
 
-		private async Task<int> FindLastEditInFileHistory(PerforceConnection Perforce, string FilePath, ILogger Logger)
+		private static async Task<int> FindLastEditInFileHistory(PerforceConnection Perforce, string FilePath, ILogger Logger)
 		{
 			string[] Files = { FilePath };
 			List<FileLogRecord> FileLogRecords = await Perforce.FileLogAsync(FileLogOptions.None, Files, CancellationToken.None);
@@ -29,8 +29,6 @@ namespace P4VUtils.Commands
 				return 1;
 			}
 
-
-
 			// the first file log record is what we're interested in
 			// because we're not following integrations, it should be the only one anyway
 			FileLogRecord FileLogRecord = FileLogRecords.First();
