Fix missing parameters to log
#rnx

diff --git a/Engine/Extras/P4VUtils/Commands/SubmitAndVirtualizeCommand.cs b/Engine/Extras/P4VUtils/Commands/SubmitAndVirtualizeCommand.cs
--- a/Engine/Extras/P4VUtils/Commands/SubmitAndVirtualizeCommand.cs
+++ b/Engine/Extras/P4VUtils/Commands/SubmitAndVirtualizeCommand.cs
@@ -677,7 +677,7 @@ namespace P4VUtils.Commands
 					}
 					else if (pluginFiles.Length > 1)
 					{
-						logger.LogWarning("Found multiple .uplugin files for '{Path}' at '{PluginDir}'", pluginDirectory);
+						logger.LogWarning("Found multiple .uplugin files for '{Path}' at '{PluginDir}'", packagePath, pluginDirectory);
 						return string.Empty;
 					}
 				}
@@ -688,7 +688,7 @@ namespace P4VUtils.Commands
 				}
 				else if (projectFiles.Length > 1)
 				{
-					logger.LogWarning("Found multiple .uproject files for '{Path}' at '{ProjectDir}'", projectDirectory);
+					logger.LogWarning("Found multiple .uproject files for '{Path}' at '{ProjectDir}'", packagePath, projectDirectory);
 					return string.Empty;
 				}
 
