Preflight command will now attempt to support import+ streams and direct the preflight for the stream the files come from

diff --git a/Engine/Extras/P4VUtils/Commands/PreflightCommand.cs b/Engine/Extras/P4VUtils/Commands/PreflightCommand.cs
--- a/Engine/Extras/P4VUtils/Commands/PreflightCommand.cs
+++ b/Engine/Extras/P4VUtils/Commands/PreflightCommand.cs
@@ -111,6 +111,50 @@ namespace P4VUtils.Commands
 				Stream = await Perforce.GetStreamAsync(Stream.Parent, false, CancellationToken.None);
 			}
 
+			// To support import+ streams get the common prefix and compare it to the Stream,
+			// if it doesn't match then we'll pivot to trying to find a stream that does
+			string CommonPrefix = DescribeRecord.Files[0].DepotFile;
+			if (DescribeRecord.Files.Count > 1)
+			{
+				// DescribeRecords conveniently come sorted alphabetically by depot file
+				// so the common prefix is the common elements of the first and last entry
+				string LastFile = DescribeRecord.Files.Last().DepotFile;
+				for (int i = 0; i < Math.Min(CommonPrefix.Length, LastFile.Length); i++)
+				{
+					if (CommonPrefix[i] != LastFile[i])
+					{
+						CommonPrefix = CommonPrefix.Substring(0, i);
+						break;
+					}
+				}
+			}
+			// If the common prefix is not the stream name then we'll try to find a valid stream out of the common prefix,
+			// but if at any point we have ambiguity we'll just fall back to the original stream again
+			if (!CommonPrefix.StartsWith(Stream.Name, StringComparison.Ordinal))
+			{
+				// Start by determining the common depot
+				int CommonStreamNameEnd = CommonPrefix.IndexOf('/', 2);
+				if (CommonStreamNameEnd != -1)
+				{
+					// Get the depot record and extract how many components a stream name has to determine how much
+					// of the common prefix to examine for a stream name
+					DepotRecord Depot = await Perforce.GetDepotAsync(CommonPrefix.Substring(2, CommonStreamNameEnd - 2));
+					int StreamDepth = Depot.GetStreamDepth();
+					for (int i = 0; i < StreamDepth && CommonStreamNameEnd != -1; i++)
+					{
+						CommonStreamNameEnd = CommonPrefix.IndexOf('/', CommonStreamNameEnd + 1);
+					}
+					if (CommonStreamNameEnd != -1)
+					{
+						StreamRecord? CommonStream = await Perforce.GetStreamAsync(CommonPrefix.Substring(0,CommonStreamNameEnd), false, CancellationToken.None);
+						if (CommonStream != null)
+						{
+							Stream = CommonStream;
+						}
+					}
+				}
+			}
+
 			return true;
 		}
 
diff --git a/Engine/Source/Programs/Shared/EpicGames.Perforce/Records/DepotRecord.cs b/Engine/Source/Programs/Shared/EpicGames.Perforce/Records/DepotRecord.cs
--- a/Engine/Source/Programs/Shared/EpicGames.Perforce/Records/DepotRecord.cs
+++ b/Engine/Source/Programs/Shared/EpicGames.Perforce/Records/DepotRecord.cs
@@ -1,6 +1,7 @@
 // Copyright Epic Games, Inc. All Rights Reserved.
 
 using System;
+using System.Linq;
 
 namespace EpicGames.Perforce
 {
@@ -28,15 +29,40 @@ namespace EpicGames.Perforce
 		public DateTime Date { get; set; }
 
 		/// <summary>
-		/// Creation time
+		/// Depot Type 
 		/// </summary>
 		[PerforceTag("Type", Optional = true)]
 		public string Type { get; set; } = String.Empty;
 
 		/// <summary>
-		/// Creation time
+		/// Depot Description
 		/// </summary>
 		[PerforceTag("Description", Optional = true)]
 		public string Description { get; set; } = String.Empty;
+
+		/// <summary>
+		/// Depot Stream Depth (number of slashes after depot that define the stream name)
+		/// A string containing either an integer value or an example path (i.e. //depot/foo/bar represents a depth of 2)
+		/// </summary>
+		[PerforceTag("StreamDepth", Optional = true)]
+		public string StreamDepthString { get;  set; } = String.Empty;
+
+		/// <summary>
+		/// Helper to get Depot Stream Depth (number of slashes after depot that define the stream name) as an integer
+		/// </summary>
+		public int GetStreamDepth()
+		{
+			int result;
+			if (Int32.TryParse(StreamDepthString, out result))
+			{
+				return result;
+			}
+			int slashCount = StreamDepthString.Count(f => f == '/');
+			if (slashCount > 0)
+			{
+				return slashCount - 2; // depth is the number of slashes not counting the leading //
+			}
+			return -1;
+		}
 	}
 }
