P4VUtils: Treat any *.json files as Horde config files. We've added a few more extensions recently, and will probably keep adding more.

diff --git a/Engine/Extras/P4VUtils/Commands/PreflightCommand.cs b/Engine/Extras/P4VUtils/Commands/PreflightCommand.cs
--- a/Engine/Extras/P4VUtils/Commands/PreflightCommand.cs
+++ b/Engine/Extras/P4VUtils/Commands/PreflightCommand.cs
@@ -383,10 +383,7 @@ namespace P4VUtils.Commands
 		/// <returns></returns>
 		static internal bool IsHordeConfigurationFile(string DepotPath)
 		{
-			string FileName = Path.GetFileName(DepotPath);
-			return (FileName.EndsWith(".stream.json", StringComparison.OrdinalIgnoreCase) || 
-					FileName.EndsWith(".project.json", StringComparison.OrdinalIgnoreCase) || 
-					FileName.Equals("globals.json", StringComparison.OrdinalIgnoreCase));
+			return DepotPath.EndsWith(".json", StringComparison.OrdinalIgnoreCase);
 		}
 
 		internal override async Task<bool> ParseArguments(string[] Args, IReadOnlyDictionary<string, string> ConfigValues, ILogger Logger)
