P4VUtils: Upgrade NuGet package to fix build error.

#rnx

diff --git a/Engine/Extras/P4VUtils/P4VUtils.csproj b/Engine/Extras/P4VUtils/P4VUtils.csproj
--- a/Engine/Extras/P4VUtils/P4VUtils.csproj
+++ b/Engine/Extras/P4VUtils/P4VUtils.csproj
@@ -74,7 +74,7 @@
   </ItemGroup>
 
   <ItemGroup>
-    <PackageReference Include="Microsoft.Extensions.Logging.Console" Version="5.0.0" />
+    <PackageReference Include="Microsoft.Extensions.Logging.Console" Version="6.0.0" />
   </ItemGroup>
 
   <ItemGroup>
