Horde: Add UgsTool.json configuration file for distributing P4VUtils through Horde.

#rnx

diff --git a/Engine/Extras/P4VUtils/UgsTool.json b/Engine/Extras/P4VUtils/UgsTool.json
new file mode 100644
--- /dev/null
+++ b/Engine/Extras/P4VUtils/UgsTool.json
@@ -0,0 +1,4 @@
+{
+	"installCommand": "p4vutils.exe install",
+	"uninstallCommand": "p4vutils.exe uninstall"
+}
