Horde: Missing project files for including UgsTool.json.

diff --git a/Engine/Extras/P4VUtils/P4VUtils.csproj b/Engine/Extras/P4VUtils/P4VUtils.csproj
--- a/Engine/Extras/P4VUtils/P4VUtils.csproj
+++ b/Engine/Extras/P4VUtils/P4VUtils.csproj
@@ -72,6 +72,12 @@
       <Link>Restricted\Commands\%(Filename)%(Extension)</Link>
     </Compile>
   </ItemGroup>
+    
+  <ItemGroup>
+    <Content Include="UgsTool.json">
+      <CopyToOutputDirectory>Always</CopyToOutputDirectory>
+    </Content>
+  </ItemGroup>
 
   <ItemGroup>
     <PackageReference Include="Microsoft.Extensions.Logging.Console" Version="6.0.0" />
diff --git a/Engine/Source/Programs/Horde/Horde/Horde.csproj b/Engine/Source/Programs/Horde/Horde/Horde.csproj
--- a/Engine/Source/Programs/Horde/Horde/Horde.csproj
+++ b/Engine/Source/Programs/Horde/Horde/Horde.csproj
@@ -39,6 +39,9 @@
       <CopyToOutputDirectory>Always</CopyToOutputDirectory>
       <CopyToPublishDirectory>Always</CopyToPublishDirectory>
     </Content>
+    <Content Include="UgsTool.json">
+      <CopyToOutputDirectory>Always</CopyToOutputDirectory>
+    </Content>
   </ItemGroup>
 
 </Project>
