Robomerge Preview command now works if there are only shelved files in the CL

diff --git a/Engine/Extras/P4VUtils/Commands/RobomergePreviewCommand.cs b/Engine/Extras/P4VUtils/Commands/RobomergePreviewCommand.cs
--- a/Engine/Extras/P4VUtils/Commands/RobomergePreviewCommand.cs
+++ b/Engine/Extras/P4VUtils/Commands/RobomergePreviewCommand.cs
@@ -43,25 +43,27 @@ namespace P4VUtils.Commands
 
 			List<OpenedRecord> OpenedRecords = await Perforce.OpenedAsync(OpenedOptions.None, Change, null, null, -1, FileSpecList.Any, CancellationToken.None).ToListAsync();
 
-			if (!OpenedRecords.Any(x => x.ClientFile.EndsWith(".branchmap.json", StringComparison.OrdinalIgnoreCase)))
+			if (OpenedRecords.Any(x => x.ClientFile.EndsWith(".branchmap.json", StringComparison.OrdinalIgnoreCase)))
 			{
-				Logger.LogError("No branchmapping.json files contained in changelist {Change}", Change);
-				return 1;
+				Logger.LogInformation("Shelving changelist {Change}", Change);
+				await Perforce.ShelveAsync(Change, ShelveOptions.Overwrite, new[] { "//..." }, CancellationToken.None);
+
+				List<DescribeRecord> Describe = await Perforce.DescribeAsync(DescribeOptions.Shelved, -1, new[] { Change }, CancellationToken.None);
+				if (Describe[0].Files.Count == 0)
+				{
+					Logger.LogError("No files were shelved in {Change}", Change);
+					return 1;
+				}
 			}
-
-			if (OpenedRecords.Count == 0)
+			else
 			{
-				Logger.LogError("No files checked out in changelist {Change}", Change);
-			}
+				DescribeRecord ChangeDesc = await Perforce.DescribeAsync(DescribeOptions.Shelved, -1, Change);
 
-			Logger.LogInformation("Shelving changelist {Change}", Change);
-			await Perforce.ShelveAsync(Change, ShelveOptions.Overwrite, new[] { "//..." }, CancellationToken.None);
-
-			List<DescribeRecord> Describe = await Perforce.DescribeAsync(DescribeOptions.Shelved, -1, new[] { Change }, CancellationToken.None);
-			if (Describe[0].Files.Count == 0)
-			{
-				Logger.LogError("No files were shelved in {Change}", Change);
-				return 1;
+				if (!ChangeDesc.Files.Any(x => x.DepotFile.EndsWith(".branchmap.json", StringComparison.OrdinalIgnoreCase)))
+				{
+					Logger.LogError("No branchmapping.json files contained in changelist {Change}", Change);
+					return 1;
+				}
 			}
 
 			string Url = GetUrl(Change, ConfigValues);
