Remove #lockdown from changelist description of backout commands to (try to) require the approvers to sign off on backed out changelists.
#rnx

diff --git a/Engine/Extras/P4VUtils/Commands/BackoutCommand.cs b/Engine/Extras/P4VUtils/Commands/BackoutCommand.cs
--- a/Engine/Extras/P4VUtils/Commands/BackoutCommand.cs
+++ b/Engine/Extras/P4VUtils/Commands/BackoutCommand.cs
@@ -210,6 +210,10 @@ namespace P4VUtils.Commands
 
 			InfoRecord Info = await Perforce.GetInfoAsync(InfoOptions.None, CancellationToken.None);
 
+			// Sanitize the existing description of tags we don't want to retain
+			string[] TagsToRemove = { "#lockdown" };
+			ExistingChangeRecord.Description = string.Join("\n", ExistingChangeRecord.Description.Split("\n").Where(DescLine => !TagsToRemove.Any(Tag => DescLine.StartsWith(Tag, StringComparison.OrdinalIgnoreCase))));
+
 			// Create a new CL
 			ChangeRecord NewChangeRecord = new ChangeRecord();
 			NewChangeRecord.User = Info.UserName;
