P4VUtils: Copy application files to a well known location when installing. This allows installations of the tool from other branches to replace existing installations, and allows P4V to continue working if the folder changes until restarted.

Also add a config file for Unreal Toolbox.

#rnx

diff --git a/Engine/Extras/P4VUtils/Program.cs b/Engine/Extras/P4VUtils/Program.cs
--- a/Engine/Extras/P4VUtils/Program.cs
+++ b/Engine/Extras/P4VUtils/Program.cs
@@ -121,12 +121,12 @@ namespace P4VUtils
 			else if (Args[0].Equals("install", StringComparison.OrdinalIgnoreCase))
 			{
 				Logger.LogInformation("Adding custom tools...");
-				return await UpdateCustomToolRegistration(true, Logger);
+				return await UpdateCustomToolRegistration(true, Args.Any(x => x.Equals("-inplace", StringComparison.OrdinalIgnoreCase)), Logger);
 			}
 			else if (Args[0].Equals("uninstall", StringComparison.OrdinalIgnoreCase))
 			{
 				Logger.LogInformation("Removing custom tools...");
-				return await UpdateCustomToolRegistration(false, Logger);
+				return await UpdateCustomToolRegistration(false, Args.Any(x => x.Equals("-inplace", StringComparison.OrdinalIgnoreCase)), Logger);
 			}
 			else if (Commands.TryGetValue(Args[0], out Command? Command))
 			{
@@ -377,6 +377,7 @@ namespace P4VUtils
 				Root.AppendChild(FolderDefinition);
 			}
 		}
+
 		static void RemoveCustomToolsFromFolders(XmlElement RootNode, FileReference ExecutableLocation, ILogger Logger)
 		{
 			XmlNodeList? CustomToolFolderList = RootNode.SelectNodes("CustomToolFolder");
@@ -411,7 +412,8 @@ namespace P4VUtils
 				}
 			}
 		}
-		public static async Task<int> UpdateCustomToolRegistration(bool bInstall, ILogger Logger)
+
+		public static async Task<int> UpdateCustomToolRegistration(bool bInstall, bool bInPlace, ILogger Logger)
 		{
 			DirectoryReference? ConfigDir = DirectoryReference.GetSpecialFolder(Environment.SpecialFolder.UserProfile);
 			if (ConfigDir == null)
@@ -444,8 +446,29 @@ namespace P4VUtils
 				Document.Load(ConfigFile.FullName);
 			}
 
+			FileReference ExecutableLocation = new FileReference(Environment.ProcessPath!);
+			if (!bInPlace)
+			{
+				DirectoryReference? InstallDir = DirectoryReference.GetSpecialFolder(Environment.SpecialFolder.LocalApplicationData);
+				if (InstallDir == null)
+				{
+					Logger.LogError("Unable to get app data folder location");
+					return 1;
+				}
 
-			FileReference? ExecutableLocation = new FileReference(Environment.ProcessPath!);
+				InstallDir = DirectoryReference.Combine(InstallDir, "Epic Games", "P4VUtils");
+
+				try
+				{
+					UpdateInstallFolder(ExecutableLocation.Directory, InstallDir, bInstall);
+					ExecutableLocation = FileReference.Combine(InstallDir, ExecutableLocation.GetFileName());
+				}
+				catch (Exception ex)
+				{
+					Logger.LogError(ex, "Unable to update install folder: {Message}", ex.Message);
+					return 1;
+				}
+			}
 
 			XmlElement? Root = Document.SelectSingleNode("CustomToolDefList") as XmlElement;
 			if (Root == null)
@@ -476,5 +499,32 @@ namespace P4VUtils
 			Logger.LogInformation("Written {ConfigFile}", ConfigFile.FullName);
 			return 0;
 		}
+
+		static void UpdateInstallFolder(DirectoryReference SourceDir, DirectoryReference TargetDir, bool bInstall)
+		{
+			DirectoryReference TempDir = DirectoryReference.Combine(TargetDir.ParentDirectory!, "~" + TargetDir.GetDirectoryName());
+			if (DirectoryReference.Exists(TempDir))
+			{
+				DirectoryReference.Delete(TempDir, true);
+			}
+
+			if (DirectoryReference.Exists(TargetDir))
+			{
+				Directory.Move(TargetDir.FullName, TempDir.FullName);
+				DirectoryReference.Delete(TempDir, true);
+			}
+
+			if (bInstall)
+			{
+				DirectoryReference.CreateDirectory(TempDir);
+				foreach (FileReference SourceFile in DirectoryReference.EnumerateFiles(SourceDir, "*", SearchOption.AllDirectories))
+				{
+					FileReference TargetFile = FileReference.Combine(TempDir, SourceFile.MakeRelativeTo(SourceDir));
+					DirectoryReference.CreateDirectory(TargetFile.Directory);
+					FileReference.Copy(SourceFile, TargetFile, true);
+				}
+				Directory.Move(TempDir.FullName, TargetDir.FullName);
+			}
+		}
 	}
 }
diff --git a/Engine/Extras/P4VUtils/Toolbox.json b/Engine/Extras/P4VUtils/Toolbox.json
new file mode 100644
--- /dev/null
+++ b/Engine/Extras/P4VUtils/Toolbox.json
@@ -0,0 +1,14 @@
+{
+    "installCommand": {
+        "fileName": "p4vutils.exe",
+        "arguments": [
+            "install"
+        ]
+    },
+    "uninstallCommand": {
+        "fileName": "p4vutils.exe",
+        "arguments": [
+            "uninstall"
+        ]
+    }
+}
\ No newline at end of file
