P4VUtils: Copy the toolbox.json file to the output directory.

#rnx

diff --git a/Engine/Extras/P4VUtils/P4VUtils.csproj b/Engine/Extras/P4VUtils/P4VUtils.csproj
--- a/Engine/Extras/P4VUtils/P4VUtils.csproj
+++ b/Engine/Extras/P4VUtils/P4VUtils.csproj
@@ -72,8 +72,15 @@
       <Link>Restricted\Commands\%(Filename)%(Extension)</Link>
     </Compile>
   </ItemGroup>
+  
+  <ItemGroup>
+    <None Remove="Toolbox.json" />
+  </ItemGroup>
     
   <ItemGroup>
+    <Content Include="Toolbox.json">
+      <CopyToOutputDirectory>Always</CopyToOutputDirectory>
+    </Content>
     <Content Include="UgsTool.json">
       <CopyToOutputDirectory>Always</CopyToOutputDirectory>
     </Content>
