P4VUtils: Print a message showing the directory that the app is installed to.

#rnx

diff --git a/Engine/Extras/P4VUtils/Program.cs b/Engine/Extras/P4VUtils/Program.cs
--- a/Engine/Extras/P4VUtils/Program.cs
+++ b/Engine/Extras/P4VUtils/Program.cs
@@ -457,6 +457,7 @@ namespace P4VUtils
 				}
 
 				InstallDir = DirectoryReference.Combine(InstallDir, "Epic Games", "P4VUtils");
+				Logger.LogInformation("Copying application files to {Dir}", InstallDir);
 
 				try
 				{
