Update P4VUtils to NET 8.

#rnx

diff --git a/Engine/Extras/P4VUtils/P4VUtils.csproj b/Engine/Extras/P4VUtils/P4VUtils.csproj
--- a/Engine/Extras/P4VUtils/P4VUtils.csproj
+++ b/Engine/Extras/P4VUtils/P4VUtils.csproj
@@ -26,7 +26,7 @@
     <Choose>
     <When Condition="'$(IsWindows)'=='true'">
       <PropertyGroup>
-        <TargetFramework>net6.0-windows</TargetFramework>
+        <TargetFramework>net8.0-windows</TargetFramework>
         <UseWindowsForms>true</UseWindowsForms>
         <UseWPF>true</UseWPF>
         <DefineConstants>$(DefineConstants);IS_WINDOWS</DefineConstants>
@@ -43,7 +43,7 @@
     </When>
     <Otherwise>
       <PropertyGroup>
-        <TargetFramework>net6.0</TargetFramework>
+        <TargetFramework>net8.0</TargetFramework>
       </PropertyGroup>
       <Choose>
       <When Condition="'$(IsOSX)'=='true'">
