P4VUtils: Fix package downgrade causing build failure.

#rnx

diff --git a/Engine/Extras/P4VUtils/P4VUtils.csproj b/Engine/Extras/P4VUtils/P4VUtils.csproj
--- a/Engine/Extras/P4VUtils/P4VUtils.csproj
+++ b/Engine/Extras/P4VUtils/P4VUtils.csproj
@@ -87,7 +87,7 @@
   </ItemGroup>
 
   <ItemGroup>
-    <PackageReference Include="Microsoft.Extensions.Logging.Console" Version="6.0.0" />
+    <PackageReference Include="Microsoft.Extensions.Logging.Console" Version="8.0.1" />
   </ItemGroup>
 
   <ItemGroup>
