Pass ClientName in to Backout to avoid environment having to be set up just so
#rnx

diff --git a/Engine/Extras/P4VUtils/Commands/BackoutCommand.cs b/Engine/Extras/P4VUtils/Commands/BackoutCommand.cs
--- a/Engine/Extras/P4VUtils/Commands/BackoutCommand.cs
+++ b/Engine/Extras/P4VUtils/Commands/BackoutCommand.cs
@@ -25,7 +25,7 @@ namespace P4VUtils.Commands
 	{
 		public override string Description => "P4 Admin sanctioned method of backing out a CL";
 
-		public override CustomToolInfo CustomTool => new CustomToolInfo("Safe Backout tool", "%S") { ShowConsole = true };
+		public override CustomToolInfo CustomTool => new CustomToolInfo("Safe Backout tool", "%S $c") { ShowConsole = true };
 
 		private readonly int MAX_ERROR_LENGTH = 2048;
 
@@ -117,15 +117,20 @@ namespace P4VUtils.Commands
 				Logger.LogError("Missing changelist number");
 				return 1;
 			}
-			else if (!int.TryParse(Args[1], out Change))
+			if (!int.TryParse(Args[1], out Change))
 			{
 				Logger.LogError("'{Argument}' is not a numbered changelist", Args[1]);
 				return 1;
 			}
+			string? ClientName = null;
+			if (Args.Length >= 2)
+			{
+				ClientName = Args[2];
+			}
 
 			bool Debug = Args.Any(x => x.Equals("-Debug", StringComparison.OrdinalIgnoreCase));
 
-			using PerforceConnection Perforce = new PerforceConnection(null, null, Logger);
+			using PerforceConnection Perforce = new PerforceConnection(null, null, ClientName, Logger);
 
 			BackoutValidationResult ValidationResult = await ValidateBackoutSafety(Perforce, Change, Logger);
 
