Temporarily disabling the code which changes the install folder for custom tools.
This will be replaced with a proper upgrade path.

diff --git a/Engine/Extras/P4VUtils/Program.cs b/Engine/Extras/P4VUtils/Program.cs
--- a/Engine/Extras/P4VUtils/Program.cs
+++ b/Engine/Extras/P4VUtils/Program.cs
@@ -447,29 +447,29 @@ namespace P4VUtils
 			}
 
 			FileReference ExecutableLocation = new FileReference(Environment.ProcessPath!);
-			if (!bInPlace)
-			{
-				DirectoryReference? InstallDir = DirectoryReference.GetSpecialFolder(Environment.SpecialFolder.LocalApplicationData);
-				if (InstallDir == null)
-				{
-					Logger.LogError("Unable to get app data folder location");
-					return 1;
-				}
-
-				InstallDir = DirectoryReference.Combine(InstallDir, "Epic Games", "P4VUtils");
-				Logger.LogInformation("Copying application files to {Dir}", InstallDir);
-
-				try
-				{
-					UpdateInstallFolder(ExecutableLocation.Directory, InstallDir, bInstall);
-					ExecutableLocation = FileReference.Combine(InstallDir, ExecutableLocation.GetFileName());
-				}
-				catch (Exception ex)
-				{
-					Logger.LogError(ex, "Unable to update install folder: {Message}", ex.Message);
-					return 1;
-				}
-			}
+			//if (!bInPlace)
+			//{
+			//	DirectoryReference? InstallDir = DirectoryReference.GetSpecialFolder(Environment.SpecialFolder.LocalApplicationData);
+			//	if (InstallDir == null)
+			//	{
+			//		Logger.LogError("Unable to get app data folder location");
+			//		return 1;
+			//	}
+
+			//	InstallDir = DirectoryReference.Combine(InstallDir, "Epic Games", "P4VUtils");
+			//	Logger.LogInformation("Copying application files to {Dir}", InstallDir);
+
+			//	try
+			//	{
+			//		UpdateInstallFolder(ExecutableLocation.Directory, InstallDir, bInstall);
+			//		ExecutableLocation = FileReference.Combine(InstallDir, ExecutableLocation.GetFileName());
+			//	}
+			//	catch (Exception ex)
+			//	{
+			//		Logger.LogError(ex, "Unable to update install folder: {Message}", ex.Message);
+			//		return 1;
+			//	}
+			//}
 
 			XmlElement? Root = Document.SelectSingleNode("CustomToolDefList") as XmlElement;
 			if (Root == null)
