Updated p4vutils to remove all exe entries in the p4v customtools.xml that match the exe filename as opposed to the full path.
This allows us to safely change the location of exe's without leaving lingering entries in the customtools.xml.

diff --git a/Engine/Extras/P4VUtils/Program.cs b/Engine/Extras/P4VUtils/Program.cs
--- a/Engine/Extras/P4VUtils/Program.cs
+++ b/Engine/Extras/P4VUtils/Program.cs
@@ -218,7 +218,7 @@ namespace P4VUtils
 			FileReference ExistingCommandPath = new FileReference(CommandElement.InnerText.StripQuoteArgument());
 
 			// if the executables matches this, then we remove it
-			if (ExistingCommandPath == ExecutableLocation)
+			if (ExistingCommandPath.GetFileName() == ExecutableLocation.GetFileName())
 			{
 				return true;
 			}
@@ -447,29 +447,30 @@ namespace P4VUtils
 			}
 
 			FileReference ExecutableLocation = new FileReference(Environment.ProcessPath!);
-			//if (!bInPlace)
-			//{
-			//	DirectoryReference? InstallDir = DirectoryReference.GetSpecialFolder(Environment.SpecialFolder.LocalApplicationData);
-			//	if (InstallDir == null)
-			//	{
-			//		Logger.LogError("Unable to get app data folder location");
-			//		return 1;
-			//	}
-
-			//	InstallDir = DirectoryReference.Combine(InstallDir, "Epic Games", "P4VUtils");
-			//	Logger.LogInformation("Copying application files to {Dir}", InstallDir);
-
-			//	try
-			//	{
-			//		UpdateInstallFolder(ExecutableLocation.Directory, InstallDir, bInstall);
-			//		ExecutableLocation = FileReference.Combine(InstallDir, ExecutableLocation.GetFileName());
-			//	}
-			//	catch (Exception ex)
-			//	{
-			//		Logger.LogError(ex, "Unable to update install folder: {Message}", ex.Message);
-			//		return 1;
-			//	}
-			//}
+
+			if (!bInPlace)
+			{
+				DirectoryReference? InstallDir = DirectoryReference.GetSpecialFolder(Environment.SpecialFolder.LocalApplicationData);
+				if (InstallDir == null)
+				{
+					Logger.LogError("Unable to get app data folder location");
+					return 1;
+				}
+
+				InstallDir = DirectoryReference.Combine(InstallDir, "Epic Games", "P4VUtils");
+				Logger.LogInformation("Copying application files to {Dir}", InstallDir);
+
+				try
+				{
+					UpdateInstallFolder(ExecutableLocation.Directory, InstallDir, bInstall);
+					ExecutableLocation = FileReference.Combine(InstallDir, ExecutableLocation.GetFileName());
+				}
+				catch (Exception ex)
+				{
+					Logger.LogError(ex, "Unable to update install folder: {Message}", ex.Message);
+					return 1;
+				}
+			}
 
 			XmlElement? Root = Document.SelectSingleNode("CustomToolDefList") as XmlElement;
 			if (Root == null)
