Fix convert to edit for new files when the file has wildcards in it
Convert to Edit will work even if there isn't a P4CONFIG file set for the working directory
#rnx

diff --git a/Engine/Extras/P4VUtils/Commands/ConvertToEditCommand.cs b/Engine/Extras/P4VUtils/Commands/ConvertToEditCommand.cs
--- a/Engine/Extras/P4VUtils/Commands/ConvertToEditCommand.cs
+++ b/Engine/Extras/P4VUtils/Commands/ConvertToEditCommand.cs
@@ -19,7 +19,7 @@ namespace P4VUtils.Commands
 	{
 		public override string Description => "Converts an integration to an edit";
 
-		public override CustomToolInfo CustomTool => new CustomToolInfo("Convert to Edit", "%p") { ShowConsole = true };
+		public override CustomToolInfo CustomTool => new CustomToolInfo("Convert to Edit", "%p $c") { ShowConsole = true };
 
 		public override async Task<int> Execute(string[] Args, IReadOnlyDictionary<string, string> ConfigValues, ILogger Logger)
 		{
@@ -34,10 +34,15 @@ namespace P4VUtils.Commands
 				Logger.LogError("'{Argument}' is not a numbered changelist", Args[1]);
 				return 1;
 			}
+			string? ClientName = null;
+			if (Args.Length >= 2)
+			{
+				ClientName = Args[2];
+			}
 
 			bool Debug = Args.Any(x => x.Equals("-Debug", StringComparison.OrdinalIgnoreCase));
 
-			using PerforceConnection Perforce = new PerforceConnection(null, null, Logger);
+			using PerforceConnection Perforce = new PerforceConnection(null, null, ClientName, Logger);
 			if(!await ConvertToEditAsync(Perforce, Change, Debug, Logger))
 			{
 				return 1;
@@ -154,7 +159,7 @@ namespace P4VUtils.Commands
 			{
 				Logger.LogInformation("Clearing the read-only flag from adds");
 
-				string[] LocalAddFiles = AddActions.Select(x => PerforceUtils.EscapePath(OpenedClientToLocalMap[x.ClientFile!])).ToArray();
+				string[] LocalAddFiles = AddActions.Select(x => OpenedClientToLocalMap[x.ClientFile!]).ToArray();
 				foreach (string LocalFile in LocalAddFiles)
 				{
 					Logger.LogInformation("    {LocalFile}", LocalFile);
@@ -172,7 +177,7 @@ namespace P4VUtils.Commands
 
 			// Perform basic actions
 			await ReopenFiles("delete", DeleteActions, OpenedClientToLocalMap, x => Perforce.DeleteAsync(Change, DeleteOptions.None, x, CancellationToken.None), Debug, Logger);
-			await ReopenFiles("add", AddActions, OpenedClientToLocalMap, x => Perforce.AddAsync(Change, x, CancellationToken.None), Debug, Logger);
+			await ReopenFiles("add", AddActions, OpenedClientToLocalMap, x => Perforce.AddAsync(Change, null, AddOptions.IncludeWildcards, x, CancellationToken.None), Debug, Logger);
 			await ReopenFiles("edit", EditActions, OpenedClientToLocalMap, x => Perforce.EditAsync(Change, x, CancellationToken.None), Debug, Logger);
 
 			// Perform move actions
@@ -244,7 +249,11 @@ namespace P4VUtils.Commands
 			{
 				Logger.LogInformation("Re-opening the following files for {Operation}:", Operation);
 
-				string[] LocalFiles = Records.Select(x => PerforceUtils.EscapePath(OpenedClientToLocalMap[x.ClientFile!])).ToArray();
+				// Add files need to keep the includes and use AddOptions.IncludeWildcards
+				string[] LocalFiles = Operation == "add" ?
+					Records.Select(x => OpenedClientToLocalMap[x.ClientFile!]).ToArray() :
+					Records.Select(x => PerforceUtils.EscapePath(OpenedClientToLocalMap[x.ClientFile!])).ToArray();
+
 				foreach (string LocalFile in LocalFiles)
 				{
 					Logger.LogInformation("    {LocalFile}", LocalFile);
