Fixed a bug where P4vUtils would fail to update due to the existing install files being read-only.

diff --git a/Engine/Extras/P4VUtils/Program.cs b/Engine/Extras/P4VUtils/Program.cs
--- a/Engine/Extras/P4VUtils/Program.cs
+++ b/Engine/Extras/P4VUtils/Program.cs
@@ -517,6 +517,12 @@ namespace P4VUtils
 			if (DirectoryReference.Exists(TargetDir))
 			{
 				Directory.Move(TargetDir.FullName, TempDir.FullName);
+
+				foreach (FileReference file in DirectoryReference.EnumerateFiles(TempDir, "*", SearchOption.AllDirectories))
+				{
+					file.ToFileInfo().Attributes = FileAttributes.Normal;
+				}
+
 				DirectoryReference.Delete(TempDir, true);
 			}
 
