Fix cherrypick not working on perforce 2024 servers
#rnx

diff --git a/Engine/Extras/P4VUtils/Commands/CherryPickCommand.cs b/Engine/Extras/P4VUtils/Commands/CherryPickCommand.cs
--- a/Engine/Extras/P4VUtils/Commands/CherryPickCommand.cs
+++ b/Engine/Extras/P4VUtils/Commands/CherryPickCommand.cs
@@ -72,7 +72,7 @@ namespace P4VUtils.Commands
 			NewChangeRecord = await Perforce.CreateChangeAsync(NewChangeRecord, CancellationToken.None);
 			Logger.LogInformation("Created pending changelist {Change}", NewChangeRecord.Number);
 
-			await Perforce.MergeAsync(MergeOptions.None, NewChangeRecord.Number, -1, SourceFileSpec, TargetFileSpec, CancellationToken.None);
+			await Perforce.MergeAsync(MergeOptions.Force, NewChangeRecord.Number, -1, SourceFileSpec, TargetFileSpec, CancellationToken.None);
 			Logger.LogInformation("Merged files into changelist {Change}", NewChangeRecord.Number);
 
 			PerforceResponseList<ResolveRecord> ResolveRecords = await Perforce.TryResolveAsync(NewChangeRecord.Number, ResolveOptions.Automatic, FileSpecList.Any, CancellationToken.None);
