Safe Restore tool (a mirror of Safe Backout)
- Backout of a backout suggests a Restore instead
- Restores only work on backouts
- Original safe backout functionality preserved
- Restore removes more tags, including the backout marks

diff --git a/Engine/Extras/P4VUtils/Commands/BackoutCommand.cs b/Engine/Extras/P4VUtils/Commands/BackoutCommand.cs
--- a/Engine/Extras/P4VUtils/Commands/BackoutCommand.cs
+++ b/Engine/Extras/P4VUtils/Commands/BackoutCommand.cs
@@ -20,16 +20,20 @@ namespace P4VUtils.Commands
 		public List<String> Errors = new List<String>();
 	}
 
-	[Command("backout", CommandCategory.Toolbox, 0)]
-	class BackoutCommand : Command
+	abstract class BaseBackoutCommand : Command
 	{
-		public override string Description => "P4 Admin sanctioned method of backing out a CL";
+		private readonly int MAX_ERROR_LENGTH = 2048;
 
-		public override CustomToolInfo CustomTool => new CustomToolInfo("Safe Backout tool", "%S $c") { ShowConsole = true };
+		public abstract Task<BackoutValidationResult> ExtraValidation(PerforceConnection Perforce, int Change, ILogger Logger);
 
-		private readonly int MAX_ERROR_LENGTH = 2048;
+		public virtual string GenerateDescription(DescribeRecord OldChangelist)
+		{
+			// Sanitize the existing description of tags we don't want to retain
+			string[] TagsToRemove = { "#lockdown" };
+			return string.Join("\n", OldChangelist.Description.Split("\n").Where(DescLine => !TagsToRemove.Any(Tag => DescLine.StartsWith(Tag, StringComparison.OrdinalIgnoreCase))));
+		}
 
-		private static async Task<BackoutValidationResult> ValidateBackoutSafety(PerforceConnection Perforce, int Change, ILogger Logger)
+		private async Task<BackoutValidationResult> ValidateBackoutSafety(PerforceConnection Perforce, int Change, ILogger Logger)
 		{
 			BackoutValidationResult Result = new BackoutValidationResult();
 
@@ -81,15 +85,15 @@ namespace P4VUtils.Commands
 
 				List<FStatRecord> StatRecord = await Perforce.FStatAsync(FileSpecs).ToListAsync();
 
-				if(BinaryFiles.Count == StatRecord.Count)
+				if (BinaryFiles.Count == StatRecord.Count)
 				{
-					for( int Index = 0; Index < BinaryFiles.Count; ++Index)
+					for (int Index = 0; Index < BinaryFiles.Count; ++Index)
 					{
 						if (StatRecord[Index].HeadRevision != BinaryFiles[Index].Revision)
 						{
 							Logger.LogError("\nIn CL {Change}, the file below is a binary file and revision #{Rev} isn't the head revision (#{HeadRev}),\ntherefore it isn't safe to back it out\n\n{FileName}",
 								Change, BinaryFiles[Index].Revision, StatRecord[Index].HeadRevision, BinaryFiles[Index].DepotFile);
-								Result.Errors.Add($"\nIn CL {Change}, the file below is a binary file and revision #{BinaryFiles[Index].Revision} isn't the head revision (#{StatRecord[Index].HeadRevision}), therefore it isn't safe to back it out\n\n{BinaryFiles[Index].DepotFile}");
+							Result.Errors.Add($"\nIn CL {Change}, the file below is a binary file and revision #{BinaryFiles[Index].Revision} isn't the head revision (#{StatRecord[Index].HeadRevision}), therefore it isn't safe to back it out\n\n{BinaryFiles[Index].DepotFile}");
 						}
 					}
 				}
@@ -100,6 +104,9 @@ namespace P4VUtils.Commands
 				}
 			}
 
+			BackoutValidationResult ExtraResults = await ExtraValidation(Perforce, Change, Logger);
+			Result.Errors.AddRange(ExtraResults.Errors);
+
 			if (Result.Errors.Count == 0)
 			{
 				Logger.LogInformation("CL {Change} seems safe to backout", Change);
@@ -108,7 +115,7 @@ namespace P4VUtils.Commands
 		}
 
 
-		
+
 		public override async Task<int> Execute(string[] Args, IReadOnlyDictionary<string, string> ConfigValues, ILogger Logger)
 		{
 			int Change;
@@ -138,7 +145,7 @@ namespace P4VUtils.Commands
 			{
 				Logger.LogInformation("\r\nCL {Change} isn't safe to backout, please confirm if you want to proceed.\r\n", Change);
 
-				string ErrorString = String.Join("\r\n", ValidationResult.Errors); 
+				string ErrorString = String.Join("\r\n", ValidationResult.Errors);
 				if (ErrorString.Length > MAX_ERROR_LENGTH)
 				{
 					ErrorString = $"{ErrorString.Substring(0, MAX_ERROR_LENGTH)}\r\n (...)";
@@ -171,7 +178,7 @@ namespace P4VUtils.Commands
 					Logger.LogInformation("\r\nProceeding with backout operation.");
 				}
 			}
-			
+
 			DescribeRecord ExistingChangeRecord = await Perforce.DescribeAsync(Change, CancellationToken.None);
 
 			// P4V undo checks for files opened before executing 'undo'
@@ -215,15 +222,12 @@ namespace P4VUtils.Commands
 
 			InfoRecord Info = await Perforce.GetInfoAsync(InfoOptions.None, CancellationToken.None);
 
-			// Sanitize the existing description of tags we don't want to retain
-			string[] TagsToRemove = { "#lockdown" };
-			ExistingChangeRecord.Description = string.Join("\n", ExistingChangeRecord.Description.Split("\n").Where(DescLine => !TagsToRemove.Any(Tag => DescLine.StartsWith(Tag, StringComparison.OrdinalIgnoreCase))));
-
 			// Create a new CL
 			ChangeRecord NewChangeRecord = new ChangeRecord();
 			NewChangeRecord.User = Info.UserName;
 			NewChangeRecord.Client = Info.ClientName;
-			NewChangeRecord.Description = $"[Backout] - CL{Change}\n#fyi {ExistingChangeRecord.User}\n#submittool safebackout\n#rnx\nOriginal CL Desc\n-----------------------------------------------------------------\n{ExistingChangeRecord.Description.TrimEnd()}\n";
+			NewChangeRecord.Description = GenerateDescription(ExistingChangeRecord);
+			
 			NewChangeRecord = await Perforce.CreateChangeAsync(NewChangeRecord, CancellationToken.None);
 
 			Logger.LogInformation("Created pending changelist {Change}", NewChangeRecord.Number);
@@ -264,7 +268,7 @@ namespace P4VUtils.Commands
 			}
 
 			// Convert the undo CL over to an edit.
-			if(!await ConvertToEditCommand.ConvertToEditAsync(Perforce, NewChangeRecord.Number, Debug, Logger))
+			if (!await ConvertToEditCommand.ConvertToEditAsync(Perforce, NewChangeRecord.Number, Debug, Logger))
 			{
 				return 1;
 			}
@@ -272,6 +276,115 @@ namespace P4VUtils.Commands
 			return 0;
 		}
 
-	}
+	} //class BaseBackoutCommand
+
+	[Command("backout", CommandCategory.Toolbox, 0)]
+	class BackoutCommand : BaseBackoutCommand
+	{
+		public override string Description => "P4 Admin sanctioned method of backing out a CL";
+
+		public override CustomToolInfo CustomTool => new CustomToolInfo("Safe Backout tool", "%S $c") { ShowConsole = true };
+
+		public override async Task<BackoutValidationResult> ExtraValidation(PerforceConnection Perforce, int Change, ILogger Logger)
+		{
+			BackoutValidationResult Result = new BackoutValidationResult();
+
+			Logger.LogInformation("\n\nExamining CL {Change} for extra unwanted Backout-related marks\n", Change);
+
+			// was this change cherry picked using ushell or p4vutil or was it robomerged?
+			DescribeRecord DescribeRec = await Perforce.DescribeAsync(Change, CancellationToken.None);
+			if (DescribeRec.Description.Contains("[Backout]", StringComparison.InvariantCultureIgnoreCase))
+			{
+				// Don't allow backing out of backouts, demand a restore operation instead
+				Logger.LogError("CL {Change} was a backout. To restore it, use the Safe Restore tool!", Change);
+				Result.Errors.Add(string.Format($"CL {Change} was a backout. To restore it, use the Safe Restore tool!", Change));
+			}
+
+			return Result;
+		}
+
+		public override string GenerateDescription(DescribeRecord OldChangelist)
+		{
+			string OldDescription = base.GenerateDescription(OldChangelist);
+
+			return $"[Backout] - CL{OldChangelist.Number}\n#fyi {OldChangelist.User}\n#submittool safebackout\n#rnx\nOriginal CL Desc\n-----------------------------------------------------------------\n{OldDescription.TrimEnd()}\n";
+		}
 
-}
+	} // class BackoutCommand
+
+	[Command("restore", CommandCategory.Toolbox, 1)]
+	class RestoreCommand : BaseBackoutCommand
+	{
+		public override string Description => "P4 Admin sanctioned method of restoring a previously backed-out CL";
+
+		public override CustomToolInfo CustomTool => new CustomToolInfo("Safe Restore tool", "%S $c") { ShowConsole = true };
+
+		public override async Task<BackoutValidationResult> ExtraValidation(PerforceConnection Perforce, int Change, ILogger Logger)
+		{
+			BackoutValidationResult Result = new BackoutValidationResult();
+
+			Logger.LogInformation("\n\nExamining CL {Change} for Restore-related marks\n", Change);
+
+			// Don't allow restoring of non-backouts
+			DescribeRecord DescribeRec = await Perforce.DescribeAsync(Change, CancellationToken.None);
+			if (!DescribeRec.Description.Contains("[Backout]", StringComparison.InvariantCultureIgnoreCase))
+			{
+				Logger.LogError("CL {Change} was not a backout, so doing a Restore makes no sense.", Change);
+				Result.Errors.Add(string.Format($"CL {Change} was not a backout, so doing a Restore makes no sense.", Change));
+			}
+
+			return Result;
+		}
+
+		public override string GenerateDescription(DescribeRecord OldChangelist)
+		{
+			string OldDescription = base.GenerateDescription(OldChangelist);
+
+
+			string[] TagsToIgnore = { "#changelist validated", "#rb", "#rnx", "#submittool", "#preflight", "[FYI]", "[Review]", "#robomerge", "#okforgithub"  };
+			var TagsToRename = new Dictionary<string, string>();
+			TagsToRename["#review"] = "#original-review";
+
+			List<string> NewLines = new List<string>();
+			string[] OldLines = OldDescription.Split('\n');
+			for (int i=0; i < OldLines.Length; i++)
+			{
+				string Line = OldLines[i].Trim();
+
+				// Find and skip potentially multiple subsequent [Backout] blocks 
+				if (Line.StartsWith("[Backout]", StringComparison.InvariantCultureIgnoreCase))
+				{
+					// Look for a nearby "-----" terminator
+					int skip = 1;
+					for (; skip < 7 && i + skip < OldLines.Length; ++skip)
+					{
+						if (OldLines[i + skip].Trim().StartsWith("------", StringComparison.InvariantCultureIgnoreCase))
+						{
+							break;
+						}
+					}
+					i += skip;
+					continue;
+				}
+
+				if (TagsToIgnore.Any(Tag => Line.StartsWith(Tag, StringComparison.InvariantCultureIgnoreCase)))
+				{
+					continue;
+				}
+
+				foreach (var pair in TagsToRename)
+				{
+					if (Line.StartsWith(pair.Key, StringComparison.InvariantCultureIgnoreCase))
+					{
+						Line = string.Concat(pair.Value, Line.AsSpan(pair.Key.Length));
+						break;
+					}
+				}
+
+				NewLines.Add(Line);
+			}
+
+			return string.Join("\n", NewLines);
+		}
+	} // class RestoreCommand
+} //namespace P4VUtils.Commands
