Safe Restore tool automatically creates a Swarm review when restoring a code CL. This makes sure that any subsequent fixes can be diffed against the state in which we were before the backout (that state was already reviewed) so we can see and better understand the post-backout-and-restore fixes.

diff --git a/Engine/Extras/P4VUtils/Commands/BackoutCommand.cs b/Engine/Extras/P4VUtils/Commands/BackoutCommand.cs
--- a/Engine/Extras/P4VUtils/Commands/BackoutCommand.cs
+++ b/Engine/Extras/P4VUtils/Commands/BackoutCommand.cs
@@ -6,15 +6,94 @@ using Microsoft.Extensions.Logging;
 using P4VUtils.Perforce;
 using System;
 using System.Collections.Generic;
-using System.ComponentModel;
+using System.Diagnostics;
 using System.Linq;
+using System.Net.Http;
 using System.Text;
+using System.Text.RegularExpressions;
 using System.Threading;
 using System.Threading.Tasks;
-using System.Windows;
 
 namespace P4VUtils.Commands
 {
+
+	// In our P4 API, we internally launch p4.exe with -G which neatly returns all results as a serialized Python dictionary,
+	// and there's some clever reflection code that automatically maps this serialized data to our structures so adding new fields or supporting new commands is a breeze.
+	// It's great that all p4 commands support this!
+	//
+	// Well, all except p4 -G tickets
+	//
+	// So, back to text parsing for that one
+
+	class TicketInfo
+	{
+		public string? ServerCluster { get; set; }
+		public string? Username { get; set; }
+		public string? Ticket { get; set; }
+	}
+
+	// Basic wrapper over launching p4.exe and parsing the text output
+	class RawP4
+	{
+		public static async Task<List<string>> Execute(string arguments)
+		{
+			var outputLines = new List<string>();
+			ProcessStartInfo processStartInfo = new ProcessStartInfo
+			{
+				FileName = PerforceChildProcess.GetExecutable(),
+				Arguments = arguments,
+				RedirectStandardOutput = true,
+				UseShellExecute = false,
+				CreateNoWindow = true,
+			};
+
+			using (Process process = new Process { StartInfo = processStartInfo })
+			{
+				var tcs = new TaskCompletionSource<bool>();
+				process.EnableRaisingEvents = true;
+				process.Exited += (sender, args) => tcs.TrySetResult(true);
+
+				process.OutputDataReceived += (sender, e) =>
+				{
+					if (e.Data != null)
+					{
+						lock (outputLines)
+						{
+							outputLines.Add(e.Data);
+						}
+					}
+				};
+
+				process.Start();
+				process.BeginOutputReadLine();
+				await tcs.Task;
+			}
+			return outputLines;
+		}
+
+		public static async Task<List<TicketInfo>> Tickets()
+		{
+			List<TicketInfo> Results = new List<TicketInfo>();
+
+			List<string> RawTickets = await Execute("tickets");
+			
+			Regex TicketRegex = new Regex(@"^[^:]+:(?<cluster>[^ ]+) \((?<username>[^)]+)\) (?<ticket>[A-F0-9]+)$", RegexOptions.Compiled);
+
+			foreach (string RawTicket in RawTickets)
+			{
+				Match m = TicketRegex.Match(RawTicket.Trim());
+				TicketInfo ti = new TicketInfo { 
+					ServerCluster = m.Groups["cluster"].Value,
+					Username      = m.Groups["username"].Value,
+					Ticket        = m.Groups["ticket"].Value,
+				};
+				Results.Add(ti);
+			}
+			return Results;
+		}
+	} // class RawP4
+
+
 	class BackoutValidationResult
 	{
 		public List<String> Errors = new List<String>();
@@ -26,6 +105,8 @@ namespace P4VUtils.Commands
 
 		public abstract Task<BackoutValidationResult> ExtraValidation(PerforceConnection Perforce, int Change, ILogger Logger);
 
+		public abstract Task PostBackoutOperation(PerforceConnection Perforce, ChangeRecord Change, ILogger Logger);
+
 		public virtual string GenerateDescription(DescribeRecord OldChangelist)
 		{
 			// Sanitize the existing description of tags we don't want to retain
@@ -273,6 +354,8 @@ namespace P4VUtils.Commands
 				return 1;
 			}
 
+			await PostBackoutOperation(Perforce, NewChangeRecord, Logger);
+
 			return 0;
 		}
 
@@ -310,6 +393,11 @@ namespace P4VUtils.Commands
 			return $"[Backout] - CL{OldChangelist.Number}\n#fyi {OldChangelist.User}\n#submittool safebackout\n#rnx\nOriginal CL Desc\n-----------------------------------------------------------------\n{OldDescription.TrimEnd()}\n";
 		}
 
+		public override async Task PostBackoutOperation(PerforceConnection Perforce, ChangeRecord Change, ILogger Logger)
+		{
+			await Task.CompletedTask;
+		}
+
 	} // class BackoutCommand
 
 	[Command("restore", CommandCategory.Toolbox, 1)]
@@ -341,13 +429,13 @@ namespace P4VUtils.Commands
 			string OldDescription = base.GenerateDescription(OldChangelist);
 
 
-			string[] TagsToIgnore = { "#changelist validated", "#rb", "#rnx", "#submittool", "#preflight", "[FYI]", "[Review]", "#robomerge", "#okforgithub"  };
+			string[] TagsToIgnore = { "#changelist validated", "#rb", "#rnx", "#submittool", "#preflight", "[FYI]", "[Review]", "#robomerge", "#okforgithub" };
 			var TagsToRename = new Dictionary<string, string>();
 			TagsToRename["#review"] = "#original-review";
 
 			List<string> NewLines = new List<string>();
 			string[] OldLines = OldDescription.Split('\n');
-			for (int i=0; i < OldLines.Length; i++)
+			for (int i = 0; i < OldLines.Length; i++)
 			{
 				string Line = OldLines[i].Trim();
 
@@ -384,7 +472,106 @@ namespace P4VUtils.Commands
 				NewLines.Add(Line);
 			}
 
+			NewLines.Add($"#was-backout {OldChangelist.Number}");
 			return string.Join("\n", NewLines);
 		}
+
+		public override async Task PostBackoutOperation(PerforceConnection Perforce, ChangeRecord Change, ILogger Logger)
+		{
+			// Re-get the list of opened files in the CL along with their metadata
+			Logger.LogInformation("Getting the list of files in changelist {Change}...", Change.Number);
+			List<OpenedRecord> OpenedRecords = await Perforce.OpenedAsync(OpenedOptions.None, Change.Number, null, null, -1, FileSpecList.Any, CancellationToken.None).ToListAsync();
+			if (OpenedRecords.Count == 0)
+			{
+				Logger.LogInformation("Change has no opened files. Aborting.");
+				return;
+			}
+
+			if (!OpenedRecords.Any(Record => PerforceUtils.IsCodeFile(Record.DepotFile)))
+			{
+				Logger.LogInformation("\nNo source files found in the changelist, nothing else to do.");
+				return;
+			}
+
+			string SwarmBaseUrl;
+			try
+			{
+				SwarmBaseUrl = await Perforce.GetPropertyAsync("P4.Swarm.URL");
+			}
+			catch
+			{
+				// Silently abort if there is no Swarm integration
+				return;
+			}
+
+			Logger.LogInformation("\nSource files found in the changelist, creating a shelf");
+
+			// Parse the previous review id
+			string? OriginalReviewId = null;
+			DescribeRecord ExistingChangeRecord = await Perforce.DescribeAsync(Change.Number, CancellationToken.None);
+			Regex ReviewIdRegex = new Regex(@"#original-review-(\d+)", RegexOptions.Compiled);
+			Match match = ReviewIdRegex.Match(ExistingChangeRecord.Description);
+			if (match.Success)
+			{
+				OriginalReviewId = match.Groups[1].Value;
+			}
+
+			await Perforce.ShelveAsync(Change.Number, ShelveOptions.Overwrite, new[] { "//..." }, CancellationToken.None);
+
+			// This can fail, but we don't care that much as that's an optional, cherry-on-top operation
+			Logger.LogInformation("Creating a swarm review...");
+			try
+			{
+				Logger.LogInformation("Getting P4 auth token");
+
+				InfoRecord userInfo = await Perforce.GetInfoAsync(InfoOptions.None);
+				List<TicketInfo> tickets = await RawP4.Tickets();
+				TicketInfo ticket = tickets.Where(Ticket => Ticket.ServerCluster == userInfo.ServerCluster && Ticket.Username == userInfo.UserName).First();
+				string AuthToken = Convert.ToBase64String(Encoding.ASCII.GetBytes($"{ticket.Username}:{ticket.Ticket}"));
+
+				Uri ReviewEndpoint  = new Uri(SwarmBaseUrl + "/api/v9/reviews");
+				Uri CommentEndpoint = new Uri(SwarmBaseUrl + "/api/v9/comments");
+
+				using (HttpClient client = new HttpClient())
+				{
+					client.DefaultRequestHeaders.Authorization = new System.Net.Http.Headers.AuthenticationHeaderValue("Basic", AuthToken);
+
+					string? ReviewId;
+
+					Logger.LogInformation("HTTP: Create review");
+					using (StringContent Content = new StringContent($"change={Change.Number}", Encoding.UTF8, "application/x-www-form-urlencoded"))
+					{
+						HttpResponseMessage Response = await client.PostAsync(ReviewEndpoint, Content);
+						Response.EnsureSuccessStatusCode();
+						string ResponseString = await Response.Content.ReadAsStringAsync();
+						JsonObject ResponseJson = JsonObject.Parse(ResponseString);
+						JsonObject ReviewObject = ResponseJson.GetObjectField("review");
+						ReviewId = ReviewObject.GetIntegerField("id").ToString();
+					}
+
+					if (ReviewId is not null && OriginalReviewId is not null)
+					{
+						Logger.LogInformation("HTTP: Post comment");
+						string[] CommentRequest = { 
+							$"topic=reviews/{ReviewId}",
+							$"body=Original review: {SwarmBaseUrl}/reviews/{OriginalReviewId}",
+							"silenceNotification=true"
+						};
+						using (StringContent Content = new StringContent(String.Join('&', CommentRequest), Encoding.UTF8, "application/x-www-form-urlencoded"))
+						{
+							HttpResponseMessage Response = await client.PostAsync(CommentEndpoint, Content);
+							Response.EnsureSuccessStatusCode();
+						}
+					}
+				}
+
+				Logger.LogInformation("Swarm review created");
+			}
+			catch
+			{
+				Logger.LogInformation("Review submit failed. Please manually submit your shelve to Swarm before you make any changes!");
+			}
+		}
+
 	} // class RestoreCommand
 } //namespace P4VUtils.Commands
diff --git a/Engine/Source/Programs/Shared/EpicGames.Perforce/PerforceChildProcess.cs b/Engine/Source/Programs/Shared/EpicGames.Perforce/PerforceChildProcess.cs
--- a/Engine/Source/Programs/Shared/EpicGames.Perforce/PerforceChildProcess.cs
+++ b/Engine/Source/Programs/Shared/EpicGames.Perforce/PerforceChildProcess.cs
@@ -15,7 +15,7 @@ namespace EpicGames.Perforce
 	/// <summary>
 	/// Wraps a call to a p4.exe child process, and allows reading data from it
 	/// </summary>
-	class PerforceChildProcess : IPerforceOutput, IDisposable
+	public class PerforceChildProcess : IPerforceOutput, IDisposable
 	{
 		/// <summary>
 		/// The process group
