Don't crash when no workspace is selected
#rnx

diff --git a/Engine/Extras/P4VUtils/Commands/BackoutCommand.cs b/Engine/Extras/P4VUtils/Commands/BackoutCommand.cs
--- a/Engine/Extras/P4VUtils/Commands/BackoutCommand.cs
+++ b/Engine/Extras/P4VUtils/Commands/BackoutCommand.cs
@@ -211,7 +211,7 @@ namespace P4VUtils.Commands
 				return 1;
 			}
 			string? ClientName = null;
-			if (Args.Length >= 2)
+			if (Args.Length > 2)
 			{
 				ClientName = Args[2];
 			}
