Ensure that the .inl versioning files are being warned for backout
#rnx

diff --git a/Engine/Extras/P4VUtils/Commands/BackoutCommand.cs b/Engine/Extras/P4VUtils/Commands/BackoutCommand.cs
--- a/Engine/Extras/P4VUtils/Commands/BackoutCommand.cs
+++ b/Engine/Extras/P4VUtils/Commands/BackoutCommand.cs
@@ -133,6 +133,7 @@ namespace P4VUtils.Commands
 
 			List<DescribeFileRecord> BinaryFiles = new List<DescribeFileRecord>();
 
+			Regex VersionFileRegex = new Regex(@"(/|Object|Custom)Versions?\.(h|inl)", RegexOptions.Compiled);
 			foreach (DescribeFileRecord DescFileRec in DescribeRec.Files)
 			{
 				// is this an edit/add/delete operation?
@@ -143,9 +144,7 @@ namespace P4VUtils.Commands
 				}
 
 				// is it an versioning file?
-				if (DescFileRec.DepotFile.Contains("ObjectVersion.h", StringComparison.InvariantCultureIgnoreCase)
-					|| DescFileRec.DepotFile.Contains("CustomVersion.h", StringComparison.InvariantCultureIgnoreCase)
-					|| DescFileRec.DepotFile.Contains("/Version.h", StringComparison.InvariantCultureIgnoreCase))
+				if (VersionFileRegex.IsMatch(DescFileRec.DepotFile))
 				{
 					Logger.LogError("In CL {Change}, the file below affects asset versioning,\ntherefore it isn't safe to back it out.\n\n{FileName}", Change, DescFileRec.DepotFile);
 					Result.Errors.Add($"In CL {Change}, the file below affects asset versioning, therefore it isn't safe to back it out.\n\n{DescFileRec.DepotFile}");
