[P4VUtils] Add JiraRegexHelp and JiraRegexPattern config options.

- The regex used to identify jira tags is now exposed as 'JiraRegexPattern' in the P4VUtils.ini file and can be overriden for teams that use different formats.
- The help message displayed if no jira tags could be found can also but customized by setting 'JiraRegexHelp' in the P4VUtils.ini file.

diff --git a/Engine/Extras/P4VUtils/Commands/JiraCommands.cs b/Engine/Extras/P4VUtils/Commands/JiraCommands.cs
--- a/Engine/Extras/P4VUtils/Commands/JiraCommands.cs
+++ b/Engine/Extras/P4VUtils/Commands/JiraCommands.cs
@@ -52,11 +52,11 @@ namespace P4VUtils.Commands
 			}
 
 			DescribeRecord description = response.Data;
-			MatchCollection matches = Regex.Matches(description.Description, @"(?<=#jira.*?)([\w]+-[\d]+)", RegexOptions.Multiline);
+			MatchCollection matches = Regex.Matches(description.Description, GetJiraRegexPattern(configValues), RegexOptions.Multiline);
 
 			if(matches.Count == 0)
 			{
-				string message = $"Description for {changeNumber} does not contain any valid jira issue tags (#jira ...)";
+				string message = $"Description for {changeNumber} does not contain any valid jira issue tags {GetJiraRegexHelpText(configValues)}";
 
 				logger.LogInformation("{Message}", message);
 				UserInterface.ShowSimpleDialog(message, "No Jira Issues Found", logger);
@@ -89,5 +89,27 @@ namespace P4VUtils.Commands
 
 			return BaseUrl.TrimEnd('/');
 		}
+
+		private static string GetJiraRegexPattern(IReadOnlyDictionary<string, string> ConfigValues)
+		{
+			string? JiraRegexPattern;
+			if (ConfigValues.TryGetValue("JiraRegexPattern", out JiraRegexPattern))
+			{
+				return JiraRegexPattern;
+			}
+
+			return "configure-jira-regex-pattern-in-p4vutils.ini";
+		}
+
+		private static string GetJiraRegexHelpText(IReadOnlyDictionary<string, string> ConfigValues)
+		{
+			string? JiraRegexHelp;
+			if (ConfigValues.TryGetValue("JiraRegexHelp", out JiraRegexHelp))
+			{
+				return JiraRegexHelp;
+			}
+
+			return "";
+		}
 	}
 }
\ No newline at end of file
diff --git a/Engine/Extras/P4VUtils/P4VUtils.ini b/Engine/Extras/P4VUtils/P4VUtils.ini
--- a/Engine/Extras/P4VUtils/P4VUtils.ini
+++ b/Engine/Extras/P4VUtils/P4VUtils.ini
@@ -1,4 +1,6 @@
 HordeServer=https://my-horde-server/
 RobomergeURL=https://my-robomerge-server/
 JiraServer=https://my-jira-server/browse/
+JiraRegexPattern=(?<=#jira.*?)([\w]+-[\d]+)
+JiraRegexHelp=(#jira ...)
 SafeBackoutHelpText=
\ No newline at end of file
