[P4VUtils] Add CommandDenyList config option.

- Both commands and groups can be filtered out.
- Example: CommandDenyList=submitandvirtualize,robomergepreviewgraph,robomergetrackchange,UE Horde,openpreflight

diff --git a/Engine/Extras/P4VUtils/Program.cs b/Engine/Extras/P4VUtils/Program.cs
--- a/Engine/Extras/P4VUtils/Program.cs
+++ b/Engine/Extras/P4VUtils/Program.cs
@@ -273,7 +273,7 @@ namespace P4VUtils
 			return ToolsChecked == ToolsRemoved;
 		}
 
-		static void InstallCommandsListInFolder(string FolderName, bool AddFolderToContextMenu, CommandCategory Category, XmlDocument Document, FileReference ExecutableLocation, ILogger Logger)
+		static void InstallCommandsListInFolder(string FolderName, bool AddFolderToContextMenu, CommandCategory Category, XmlDocument Document, FileReference ExecutableLocation, ILogger Logger, string? CommandDenyList)
 		{
 			// <CustomToolDefList>				// list of custom tools (top level)
 			//  < CustomToolDef >				// loose custom tool in top level
@@ -282,6 +282,12 @@ namespace P4VUtils
 			//		< CustomToolDefList >		// list of custom tools in folder
 			//		< CustomToolDef >			// definition of tool
 
+			if (CommandDenyList != null && CommandDenyList.Contains(FolderName, StringComparison.OrdinalIgnoreCase))
+			{
+				Logger.LogInformation("{CommandDenyList} contains {FolderName}", CommandDenyList, FolderName);
+				return;
+			}
+
 			// This is the top level node, there will also be a per folder node added of same name
 			XmlElement? Root = Document.SelectSingleNode("CustomToolDefList") as XmlElement;
 
@@ -304,6 +310,11 @@ namespace P4VUtils
 
 				foreach (KeyValuePair<string, Command> Pair in CategoryCommands)
 				{
+					if (CommandDenyList != null && CommandDenyList.Contains(Pair.Key, StringComparison.OrdinalIgnoreCase))
+					{
+						continue;
+					}
+
 					CustomToolInfo CustomTool = Pair.Value.CustomTool;
 
 					XmlElement ToolDef = Document.CreateElement("CustomToolDef");
@@ -487,12 +498,16 @@ namespace P4VUtils
 			// Insert new entries
 			if (bInstall)
 			{
-				InstallCommandsListInFolder("UE RootHelpers", false/*AddFolderToContextMenu*/, CommandCategory.Root, Document, ExecutableLocation, Logger);
-				InstallCommandsListInFolder("UE Content", true/*AddFolderToContextMenu*/, CommandCategory.Content, Document, ExecutableLocation, Logger);
-				InstallCommandsListInFolder("UE Toolbox", true/*AddFolderToContextMenu*/, CommandCategory.Toolbox, Document, ExecutableLocation, Logger);
-				InstallCommandsListInFolder("UE Integrate", true/*AddFolderToContextMenu*/, CommandCategory.Integrate, Document, ExecutableLocation, Logger);
-				InstallCommandsListInFolder("UE Horde", true/*AddFolderToContextMenu*/, CommandCategory.Horde, Document, ExecutableLocation, Logger);
-				InstallCommandsListInFolder("UE Browser", true/*AddFolderToContextMenu*/, CommandCategory.Browser, Document, ExecutableLocation, Logger);
+				Dictionary<string, string> ConfigValues = ReadConfig();
+				string? CommandDenyList;
+				ConfigValues.TryGetValue("CommandDenyList", out CommandDenyList);
+
+				InstallCommandsListInFolder("UE RootHelpers", false/*AddFolderToContextMenu*/, CommandCategory.Root, Document, ExecutableLocation, Logger, CommandDenyList);
+				InstallCommandsListInFolder("UE Content", true/*AddFolderToContextMenu*/, CommandCategory.Content, Document, ExecutableLocation, Logger, CommandDenyList);
+				InstallCommandsListInFolder("UE Toolbox", true/*AddFolderToContextMenu*/, CommandCategory.Toolbox, Document, ExecutableLocation, Logger, CommandDenyList);
+				InstallCommandsListInFolder("UE Integrate", true/*AddFolderToContextMenu*/, CommandCategory.Integrate, Document, ExecutableLocation, Logger, CommandDenyList);
+				InstallCommandsListInFolder("UE Horde", true/*AddFolderToContextMenu*/, CommandCategory.Horde, Document, ExecutableLocation, Logger, CommandDenyList);
+				InstallCommandsListInFolder("UE Browser", true/*AddFolderToContextMenu*/, CommandCategory.Browser, Document, ExecutableLocation, Logger, CommandDenyList);
 			}
 
 			// Save the new document
