Added `#Synced-CL` to the`TagsToRemove` in `BaseBackoutCommand` of the safe backout and safe restore tools.

#rnx

#Synced-CL 50325520

diff --git a/Engine/Extras/P4VUtils/Commands/BackoutCommand.cs b/Engine/Extras/P4VUtils/Commands/BackoutCommand.cs
--- a/Engine/Extras/P4VUtils/Commands/BackoutCommand.cs
+++ b/Engine/Extras/P4VUtils/Commands/BackoutCommand.cs
@@ -110,7 +110,7 @@ namespace P4VUtils.Commands
 		public virtual string GenerateDescription(DescribeRecord OldChangelist)
 		{
 			// Sanitize the existing description of tags we don't want to retain
-			string[] TagsToRemove = { "#lockdown" };
+			string[] TagsToRemove = { "#lockdown", "#Synced-CL" };
 			return string.Join("\n", OldChangelist.Description.Split("\n").Where(DescLine => !TagsToRemove.Any(Tag => DescLine.StartsWith(Tag, StringComparison.OrdinalIgnoreCase))));
 		}
 
