Add version number to safe backout tool's generated tag to help determine when users are using an out of date version.

#Synced-CL 50466157

diff --git a/Engine/Extras/P4VUtils/Commands/BackoutCommand.cs b/Engine/Extras/P4VUtils/Commands/BackoutCommand.cs
--- a/Engine/Extras/P4VUtils/Commands/BackoutCommand.cs
+++ b/Engine/Extras/P4VUtils/Commands/BackoutCommand.cs
@@ -389,7 +389,7 @@ namespace P4VUtils.Commands
 		{
 			string OldDescription = base.GenerateDescription(OldChangelist);
 
-			return $"[Backout] - CL{OldChangelist.Number}\n#fyi {OldChangelist.User}\n#submittool safebackout\n#rnx\nOriginal CL Desc\n-----------------------------------------------------------------\n{OldDescription.TrimEnd()}\n";
+			return $"[Backout] - CL{OldChangelist.Number}\n#fyi {OldChangelist.User}\n#submittool SafeBackout-{Program.GetBuildVersion()}\n#rnx\nOriginal CL Desc\n-----------------------------------------------------------------\n{OldDescription.TrimEnd()}\n";
 		}
 
 		public override async Task PostBackoutOperation(PerforceConnection Perforce, ChangeRecord Change, ILogger Logger)
diff --git a/Engine/Extras/P4VUtils/P4VUtils.csproj b/Engine/Extras/P4VUtils/P4VUtils.csproj
--- a/Engine/Extras/P4VUtils/P4VUtils.csproj
+++ b/Engine/Extras/P4VUtils/P4VUtils.csproj
@@ -12,6 +12,14 @@
     <NoWarn>CA1812</NoWarn>
   </PropertyGroup>
 
+  <!-- Version information - uebp_CL is set by Horde on build machines -->
+  <PropertyGroup>
+    <Version>1.0.0</Version>
+    <BuildChangelist>$(uebp_CL)</BuildChangelist>
+    <BuildChangelist Condition="'$(BuildChangelist)' == ''">NoCL</BuildChangelist>
+    <InformationalVersion>$(Version)-$(BuildChangelist)</InformationalVersion>
+  </PropertyGroup>
+
     <PropertyGroup>
       <IsWindows Condition="'$([System.Runtime.InteropServices.RuntimeInformation]::IsOSPlatform($([System.Runtime.InteropServices.OSPlatform]::Windows)))' == 'true'">true</IsWindows> 
       <IsOSX Condition="'$([System.Runtime.InteropServices.RuntimeInformation]::IsOSPlatform($([System.Runtime.InteropServices.OSPlatform]::OSX)))' == 'true'">true</IsOSX> 
diff --git a/Engine/Extras/P4VUtils/Program.cs b/Engine/Extras/P4VUtils/Program.cs
--- a/Engine/Extras/P4VUtils/Program.cs
+++ b/Engine/Extras/P4VUtils/Program.cs
@@ -44,6 +44,15 @@ namespace P4VUtils
 		public static IDictionary<string, Command> Commands = new Dictionary<string, Command>();
 		public static IDictionary<CommandCategory, List<string>> CommandCategories = new Dictionary<CommandCategory, List<string>>();
 
+		/// <summary>
+		/// Returns the P4 changelist number this build was created from, or "NoCL" if unavailable.
+		/// </summary>
+		public static string GetBuildVersion()
+		{
+			AssemblyInformationalVersionAttribute? version = Assembly.GetExecutingAssembly().GetCustomAttribute<AssemblyInformationalVersionAttribute>();
+			return version?.InformationalVersion ?? "NoCL";
+		}
+
 		static void PrintHelp(ILogger Logger)
 		{
 			Logger.LogInformation("P4VUtils");
