Fix P4VUtils for Linux to support dialogs and copy to clipboard

#rnx

#Synced-CL 52346205

diff --git a/Engine/Extras/P4VUtils/UserInterface.cs b/Engine/Extras/P4VUtils/UserInterface.cs
--- a/Engine/Extras/P4VUtils/UserInterface.cs
+++ b/Engine/Extras/P4VUtils/UserInterface.cs
@@ -101,6 +101,119 @@ namespace P4VUtils
 				Output = Output.Replace("button returned:", "", StringComparison.OrdinalIgnoreCase);
 				Response = Enum.Parse<Button>(Output);
 			}
+			else if (OperatingSystem.IsLinux())
+			{
+				// Try zenity for GUI dialogs, fall back to kdialog, then log-only
+				string? ZenityPath = null;
+				string? KDialogPath = null;
+				try
+				{
+					ZenityPath = RunProcessAsync("/usr/bin/env", "bash -c \"which zenity 2>/dev/null\"").Result;
+					if (string.IsNullOrEmpty(ZenityPath)) { ZenityPath = null; }
+				}
+				catch { }
+				if (ZenityPath == null)
+				{
+					try
+					{
+						KDialogPath = RunProcessAsync("/usr/bin/env", "bash -c \"which kdialog 2>/dev/null\"").Result;
+						if (string.IsNullOrEmpty(KDialogPath)) { KDialogPath = null; }
+					}
+					catch { }
+				}
+
+				if (ZenityPath != null)
+				{
+					string EscapedMessage = ShellEscape(Message);
+					string EscapedTitle = ShellEscape(Title);
+
+					if (Buttons.Length == 1 && Buttons[0] == Button.OK)
+					{
+						// Simple info dialog
+						RunProcessWithExitCodeAsync("/usr/bin/env", $"bash -c \"zenity --info --title=\\\"{EscapedTitle}\\\" --text=\\\"{EscapedMessage}\\\" --width=400\"").Wait();
+						Response = Button.OK;
+					}
+					else if (Buttons.SequenceEqual(YesNo))
+					{
+						// zenity: exit 0 = OK button, exit 1 = Cancel button
+						var (_, ExitCode) = RunProcessWithExitCodeAsync("/usr/bin/env", $"bash -c \"zenity --question --title=\\\"{EscapedTitle}\\\" --text=\\\"{EscapedMessage}\\\" --ok-label=Yes --cancel-label=No --width=400\"").Result;
+						Response = ExitCode == 0 ? Button.Yes : Button.No;
+					}
+					else if (Buttons.SequenceEqual(OKCancel))
+					{
+						var (_, ExitCode) = RunProcessWithExitCodeAsync("/usr/bin/env", $"bash -c \"zenity --question --title=\\\"{EscapedTitle}\\\" --text=\\\"{EscapedMessage}\\\" --ok-label=OK --cancel-label=Cancel --width=400\"").Result;
+						Response = ExitCode == 0 ? Button.OK : Button.Cancel;
+					}
+					else if (Buttons.SequenceEqual(YesNoCancel))
+					{
+						// zenity --extra-button prints button label to stdout and exits with code 1
+						// OK button (Yes): exit 0, no stdout
+						// Cancel button (No): exit 1, no stdout
+						// Extra button (Cancel): exit 1, prints "Cancel" to stdout
+						var (Output, ExitCode) = RunProcessWithExitCodeAsync("/usr/bin/env", $"bash -c \"zenity --question --title=\\\"{EscapedTitle}\\\" --text=\\\"{EscapedMessage}\\\" --ok-label=Yes --cancel-label=No --extra-button=Cancel --width=400\"").Result;
+						if (ExitCode == 0)
+						{
+							Response = Button.Yes;
+						}
+						else if (Output == "Cancel")
+						{
+							Response = Button.Cancel;
+						}
+						else
+						{
+							Response = Button.No;
+						}
+					}
+				}
+				else if (KDialogPath != null)
+				{
+					string EscapedMessage = ShellEscape(Message);
+					string EscapedTitle = ShellEscape(Title);
+
+					if (Buttons.Length == 1 && Buttons[0] == Button.OK)
+					{
+						RunProcessWithExitCodeAsync("/usr/bin/env", $"bash -c \"kdialog --title \\\"{EscapedTitle}\\\" --msgbox \\\"{EscapedMessage}\\\"\"").Wait();
+						Response = Button.OK;
+					}
+					else if (Buttons.SequenceEqual(YesNo))
+					{
+						// kdialog --yesno: exit 0 = Yes, exit 1 = No
+						var (_, ExitCode) = RunProcessWithExitCodeAsync("/usr/bin/env", $"bash -c \"kdialog --title \\\"{EscapedTitle}\\\" --yesno \\\"{EscapedMessage}\\\"\"").Result;
+						Response = ExitCode == 0 ? Button.Yes : Button.No;
+					}
+					else if (Buttons.SequenceEqual(OKCancel))
+					{
+						var (_, ExitCode) = RunProcessWithExitCodeAsync("/usr/bin/env", $"bash -c \"kdialog --title \\\"{EscapedTitle}\\\" --yesno \\\"{EscapedMessage}\\\" --yes-label OK --no-label Cancel\"").Result;
+						Response = ExitCode == 0 ? Button.OK : Button.Cancel;
+					}
+					else if (Buttons.SequenceEqual(YesNoCancel))
+					{
+						// kdialog --yesnocancel: exit 0 = Yes, exit 1 = No, exit 2 = Cancel
+						var (_, ExitCode) = RunProcessWithExitCodeAsync("/usr/bin/env", $"bash -c \"kdialog --title \\\"{EscapedTitle}\\\" --yesnocancel \\\"{EscapedMessage}\\\"\"").Result;
+						if (ExitCode == 0)
+						{
+							Response = Button.Yes;
+						}
+						else if (ExitCode == 2)
+						{
+							Response = Button.Cancel;
+						}
+						else
+						{
+							Response = Button.No;
+						}
+					}
+				}
+				else
+				{
+					Logger.LogWarning("{Message}", Message);
+					if (Buttons.Length > 1)
+					{
+						Logger.LogWarning("Note: This dialog was meant to have buttons {Buttons}, returning {Default}. Install zenity or kdialog for GUI dialogs.", string.Join(", ", Buttons), DefaultButton);
+					}
+					Response = DefaultButton;
+				}
+			}
 			else
 			{
 				Logger.LogWarning("{Message}", Message);
@@ -145,6 +258,49 @@ namespace P4VUtils
 			{
 				await UserInterface.RunProcessAsync("/bin/bash", $" -c \"osascript -e 'set the clipboard to \\\"{Text}\\\"'\"");
 			}
+			else if (OperatingSystem.IsLinux())
+			{
+				// Try xclip, xsel, then wl-copy (Wayland) in order of preference
+				string[] ClipboardCommands = new[]
+				{
+					"xclip -selection clipboard",
+					"xsel --clipboard --input",
+					"wl-copy"
+				};
+
+				bool Copied = false;
+				foreach (string ClipCmd in ClipboardCommands)
+				{
+					string[] Parts = ClipCmd.Split(' ', 2);
+					string Executable = Parts[0];
+
+					try
+					{
+						// Check if the tool exists
+						string WhichOutput = await RunProcessAsync("/usr/bin/env", $"bash -c \"which {Executable} 2>/dev/null\"");
+						if (string.IsNullOrEmpty(WhichOutput))
+						{
+							continue;
+						}
+
+						int ExitCode = await RunProcessWithStdinAsync("/usr/bin/env", $"bash -c \"{ClipCmd}\"", Text);
+						if (ExitCode == 0)
+						{
+							Copied = true;
+							break;
+						}
+					}
+					catch (Exception)
+					{
+						continue;
+					}
+				}
+
+				if (!Copied)
+				{
+					Logger.LogError("Unable to copy to clipboard on Linux. Install xclip, xsel, or wl-copy.");
+				}
+			}
 			else
 			{
 				Logger.LogError("Unable to copy to clipboard on this platform");
@@ -169,6 +325,46 @@ namespace P4VUtils
 			}
 		}
 
+		public static async Task<(string Output, int ExitCode)> RunProcessWithExitCodeAsync(string Executable, string Params)
+		{
+			ProcessStartInfo SI = new ProcessStartInfo(Executable, Params);
+			SI.RedirectStandardOutput = true;
+			SI.RedirectStandardError = true;
+			using (Process Proc = Process.Start(SI)!)
+			{
+				string Output = (await Proc.StandardOutput.ReadToEndAsync()).Trim();
+				await Proc.WaitForExitAsync();
+				return (Output, Proc.ExitCode);
+			}
+		}
+
+		public static async Task<int> RunProcessWithStdinAsync(string Executable, string Params, string StdinText)
+		{
+			ProcessStartInfo SI = new ProcessStartInfo(Executable, Params);
+			SI.RedirectStandardInput = true;
+			SI.RedirectStandardOutput = true;
+			SI.RedirectStandardError = true;
+			using (Process Proc = Process.Start(SI)!)
+			{
+				await Proc.StandardInput.WriteAsync(StdinText);
+				Proc.StandardInput.Close();
+				await Proc.WaitForExitAsync();
+				return Proc.ExitCode;
+			}
+		}
+
+		// Escape a string for safe embedding in shell double-quotes
+		private static string ShellEscape(string Input)
+		{
+			return Input
+				.Replace("\\", "\\\\", StringComparison.Ordinal)
+				.Replace("\"", "\\\"", StringComparison.Ordinal)
+				.Replace("$", "\\$", StringComparison.Ordinal)
+				.Replace("`", "\\`", StringComparison.Ordinal)
+				.Replace("\n", "\\n", StringComparison.Ordinal)
+				.Replace("\r", "", StringComparison.Ordinal);
+		}
+
 		#endregion
 	}
 }
