P4VUtils: Add missing docstrings

#rnx
#Synced-CL 52421477

diff --git a/Engine/Extras/P4VUtils/CommandAttribute.cs b/Engine/Extras/P4VUtils/CommandAttribute.cs
--- a/Engine/Extras/P4VUtils/CommandAttribute.cs
+++ b/Engine/Extras/P4VUtils/CommandAttribute.cs
@@ -4,6 +4,7 @@ using System;
 
 namespace P4VUtils
 {
+	/// <summary>Category used to organize commands in the UI</summary>
 	public enum CommandCategory
 	{
 		/// <summary>
@@ -62,6 +63,10 @@ namespace P4VUtils
 		/// </summary>
 		public int Order { get; internal set; }
 
+		/// <summary>Initializes a new instance of <see cref="CommandAttribute"/></summary>
+		/// <param name="commandName">CLI name of the command</param>
+		/// <param name="category">Category the command belongs to</param>
+		/// <param name="order">Sort order within the category</param>
 		public CommandAttribute(string commandName, CommandCategory category, int order = int.MaxValue)
 		{
 			CommandName = commandName;
diff --git a/Engine/Extras/P4VUtils/Commands/SubmitAndVirtualizeCommand.cs b/Engine/Extras/P4VUtils/Commands/SubmitAndVirtualizeCommand.cs
--- a/Engine/Extras/P4VUtils/Commands/SubmitAndVirtualizeCommand.cs
+++ b/Engine/Extras/P4VUtils/Commands/SubmitAndVirtualizeCommand.cs
@@ -26,14 +26,23 @@ namespace P4VUtils.Commands
 	/// </summary>
 	public class BuildVersion
 	{
+		/// <summary>Major version number</summary>
 		public int? MajorVersion { get; set; }
+		/// <summary>Minor version number</summary>
 		public int? MinorVersion { get; set; }
+		/// <summary>Patch version number</summary>
 		public int? PatchVersion { get; set; }
+		/// <summary>Perforce changelist number this build was made from</summary>
 		public int? Changelist { get; set; }
+		/// <summary>Oldest changelist this build is compatible with</summary>
 		public int? CompatibleChangelist { get; set; }
+		/// <summary>Whether this is a licensee build</summary>
 		public int? IsLicenseeVersion { get; set; }
+		/// <summary>Whether this is a promoted build</summary>
 		public int? IsPromotedBuild { get; set; }
+		/// <summary>Source control branch name</summary>
 		public string? BranchName { get; set; }
+		/// <summary>Unique identifier for this build</summary>
 		public string? BuildId { get; set; }
 	}
 
@@ -43,9 +52,11 @@ namespace P4VUtils.Commands
 	/// </summary>
 	public class TargetBuildVersion
 	{
+		/// <summary>Build version information</summary>
 		public BuildVersion? Version { get; set; }
 	}
 
+	/// <summary>Logging extension methods</summary>
 	public static class LoggingUtils
 	{
 		/// <summary>
diff --git a/Engine/Extras/P4VUtils/P4ActionGroups.cs b/Engine/Extras/P4VUtils/P4ActionGroups.cs
--- a/Engine/Extras/P4VUtils/P4ActionGroups.cs
+++ b/Engine/Extras/P4VUtils/P4ActionGroups.cs
@@ -4,8 +4,10 @@ using EpicGames.Perforce;
 
 namespace P4VUtils.Perforce
 {
+	/// <summary>Groups of Perforce file actions used for filtering and categorization</summary>
 	public static class P4ActionGroups
 	{
+		/// <summary>Integrate actions that bring changes from another stream/branch</summary>
 		public static readonly IntegrateAction[] IntegrateFromActions =
 		{
 			IntegrateAction.BranchFrom,
@@ -17,6 +19,7 @@ namespace P4VUtils.Perforce
 			IntegrateAction.AddFrom
 		};
 
+		/// <summary>Actions representing direct edits to a file (add, edit, delete)</summary>
 		public static readonly FileAction[] EditActions =
 		{
 			FileAction.Add,
@@ -24,12 +27,14 @@ namespace P4VUtils.Perforce
 			FileAction.Delete
 		};
 
+		/// <summary>Actions representing file moves</summary>
 		public static readonly FileAction[] MoveActions =
 		{
 			FileAction.MoveAdd,
 			FileAction.MoveDelete
 		};
 
+		/// <summary>Actions representing integrations from another branch</summary>
 		public static readonly FileAction[] IntegrateActions =
 		{
 			FileAction.Integrate,
diff --git a/Engine/Extras/P4VUtils/ProcessUtils.cs b/Engine/Extras/P4VUtils/ProcessUtils.cs
--- a/Engine/Extras/P4VUtils/ProcessUtils.cs
+++ b/Engine/Extras/P4VUtils/ProcessUtils.cs
@@ -5,6 +5,7 @@ using System.Runtime.InteropServices;
 
 namespace P4VUtils
 {
+	/// <summary>Utilities for launching processes</summary>
 	public static class ProcessUtils
 	{
 		/// <summary>
@@ -13,7 +14,7 @@ namespace P4VUtils
 		/// A .txt file will probably open a text editor, a url will open the
 		/// default browser etc.
 		/// </summary>
-		/// <param name="Path"The path to be opened></param>
+		/// <param name="Path">The path to be opened</param>
 		public static void OpenInNewProcess(string Path)
 		{
 			if (RuntimeInformation.IsOSPlatform(OSPlatform.Windows))
diff --git a/Engine/Extras/P4VUtils/UserInterface.cs b/Engine/Extras/P4VUtils/UserInterface.cs
--- a/Engine/Extras/P4VUtils/UserInterface.cs
+++ b/Engine/Extras/P4VUtils/UserInterface.cs
@@ -18,12 +18,14 @@ using System.Windows.Forms;
 
 namespace P4VUtils
 {
+	/// <summary>Cross-platform UI helpers for dialogs, clipboard, and process launching</summary>
 	public static class UserInterface
 	{
 #if IS_WINDOWS
 		[System.Runtime.InteropServices.DllImport("user32.dll")]
 		private static extern bool SetProcessDPIAware();
 
+		/// <summary>Configures DPI awareness and visual styles for the application</summary>
 		public static void SetupVisuals()
 		{
 			// make the form look good on modern displays!
@@ -33,6 +35,7 @@ namespace P4VUtils
 			Application.SetCompatibleTextRenderingDefault(false);
 		}
 #else
+		/// <summary>Configures DPI awareness and visual styles for the application</summary>
 		public static void SetupVisuals() { }
 #endif
 
@@ -44,17 +47,26 @@ namespace P4VUtils
 
 		#region Dialog Boxes
 
+		/// <summary>Dialog button identifiers</summary>
 		public enum Button
 		{
+			/// <summary>Yes button</summary>
 			Yes,
+			/// <summary>No button</summary>
 			No,
+			/// <summary>OK button</summary>
 			OK,
+			/// <summary>Cancel button</summary>
 			Cancel,
 		}
 
+		/// <summary>Button set for Yes/No/Cancel dialogs</summary>
 		public static readonly Button[] YesNoCancel = { Button.Yes, Button.No, Button.Cancel };
+		/// <summary>Button set for Yes/No dialogs</summary>
 		public static readonly Button[] YesNo = { Button.Yes, Button.No };
+		/// <summary>Button set for simple OK dialogs</summary>
 		public static readonly Button[] OK = { Button.OK };
+		/// <summary>Button set for OK/Cancel dialogs</summary>
 		public static readonly Button[] OKCancel = { Button.OK, Button.Cancel };
 
 #if IS_WINDOWS
@@ -73,11 +85,22 @@ namespace P4VUtils
 		};
 #endif
 
+		/// <summary>Shows an informational dialog with a single OK button</summary>
+		/// <param name="Message">Message to display</param>
+		/// <param name="Title">Dialog title</param>
+		/// <param name="Logger">Logger for fallback output on non-GUI platforms</param>
 		public static void ShowSimpleDialog(string Message, String Title, ILogger Logger)
 		{
 			ShowDialog(Message, Title, OK, Button.OK, Logger);
 		}
 
+		/// <summary>Shows a dialog with a configurable set of buttons and returns the button the user clicked</summary>
+		/// <param name="Message">Message to display</param>
+		/// <param name="Title">Dialog title</param>
+		/// <param name="Buttons">Buttons to present</param>
+		/// <param name="DefaultButton">Button returned when no GUI is available</param>
+		/// <param name="Logger">Logger for fallback output on non-GUI platforms</param>
+		/// <returns>The button the user clicked</returns>
 		public static Button ShowDialog(string Message, string Title, Button[] Buttons, Button DefaultButton, ILogger Logger)
 		{
 			Button Response = Button.OK;
@@ -232,6 +255,9 @@ namespace P4VUtils
 
 		#region Clipboard
 
+		/// <summary>Copies the given text to the system clipboard</summary>
+		/// <param name="Text">Text to copy</param>
+		/// <param name="Logger">Logger for error output when clipboard access fails</param>
 		public static async Task CopyTextToClipboard(string Text, ILogger Logger)
 		{
 			if (OperatingSystem.IsWindows())
@@ -312,6 +338,10 @@ namespace P4VUtils
 
 		#region Utils
 
+		/// <summary>Launches a process and returns its stdout as a trimmed string</summary>
+		/// <param name="Executable">Path to the executable</param>
+		/// <param name="Params">Command-line arguments</param>
+		/// <returns>Trimmed stdout output</returns>
 		public static async Task<string> RunProcessAsync(string Executable, string Params)
 		{
 			ProcessStartInfo SI = new ProcessStartInfo(Executable, Params);
@@ -325,6 +355,10 @@ namespace P4VUtils
 			}
 		}
 
+		/// <summary>Launches a process and returns its stdout and exit code</summary>
+		/// <param name="Executable">Path to the executable</param>
+		/// <param name="Params">Command-line arguments</param>
+		/// <returns>Trimmed stdout and the process exit code</returns>
 		public static async Task<(string Output, int ExitCode)> RunProcessWithExitCodeAsync(string Executable, string Params)
 		{
 			ProcessStartInfo SI = new ProcessStartInfo(Executable, Params);
@@ -338,6 +372,11 @@ namespace P4VUtils
 			}
 		}
 
+		/// <summary>Launches a process, writes text to its stdin, and returns the exit code</summary>
+		/// <param name="Executable">Path to the executable</param>
+		/// <param name="Params">Command-line arguments</param>
+		/// <param name="StdinText">Text to write to the process stdin</param>
+		/// <returns>The process exit code</returns>
 		public static async Task<int> RunProcessWithStdinAsync(string Executable, string Params, string StdinText)
 		{
 			ProcessStartInfo SI = new ProcessStartInfo(Executable, Params);
