P4Utils: Upgrade Microsoft.Extensions.Logging.Console package to be consistent with other post-dotnet 10-upgrade projects.

#Synced-CL 52578659

diff --git a/Engine/Extras/P4VUtils/P4VUtils.csproj b/Engine/Extras/P4VUtils/P4VUtils.csproj
--- a/Engine/Extras/P4VUtils/P4VUtils.csproj
+++ b/Engine/Extras/P4VUtils/P4VUtils.csproj
@@ -93,7 +93,7 @@
   </ItemGroup>
 
   <ItemGroup>
-    <PackageReference Include="Microsoft.Extensions.Logging.Console" Version="8.0.1" />
+    <PackageReference Include="Microsoft.Extensions.Logging.Console" Version="10.0.5" />
   </ItemGroup>
 
   <ItemGroup>
